Fixing advanced vehicle template by calling superclass functions correctly. Vehicle wasn't responding to controls or updating HUD because Super::BeginPlay() was never called.

diff --git a/Templates/TP_VehicleAdv/Source/TP_VehicleAdv/TP_VehicleAdvPawn.cpp b/Templates/TP_VehicleAdv/Source/TP_VehicleAdv/TP_VehicleAdvPawn.cpp
--- a/Templates/TP_VehicleAdv/Source/TP_VehicleAdv/TP_VehicleAdvPawn.cpp
+++ b/Templates/TP_VehicleAdv/Source/TP_VehicleAdv/TP_VehicleAdvPawn.cpp
@@ -169,6 +169,8 @@ ATP_VehicleAdvPawn::ATP_VehicleAdvPawn()
 
 void ATP_VehicleAdvPawn::SetupPlayerInputComponent(class UInputComponent* InputComponent)
 {
+	Super::SetupPlayerInputComponent(InputComponent);
+
 	// set up gameplay key bindings
 	check(InputComponent);
 
@@ -235,6 +237,8 @@ void ATP_VehicleAdvPawn::EnableIncarView(const bool bState)
 
 void ATP_VehicleAdvPawn::Tick(float Delta)
 {
+	Super::Tick(Delta);
+
 	// Setup the flag to say we are in reverse gear
 	bInReverseGear = GetVehicleMovement()->GetCurrentGear() < 0;
 	
@@ -272,6 +276,8 @@ void ATP_VehicleAdvPawn::Tick(float Delta)
 
 void ATP_VehicleAdvPawn::BeginPlay()
 {
+	Super::BeginPlay();
+
 	bool bWantInCar = false;
 	// First disable both speed/gear displays
 	bInCarCameraActive = false;
