Vehicle no longer switches gears when not moving in vehicle advanced code and BP.
Resaving layerinfo to resolve empty engine version warning

diff --git a/Templates/TP_VehicleAdv/Source/TP_VehicleAdv/TP_VehicleAdvPawn.cpp b/Templates/TP_VehicleAdv/Source/TP_VehicleAdv/TP_VehicleAdvPawn.cpp
--- a/Templates/TP_VehicleAdv/Source/TP_VehicleAdv/TP_VehicleAdvPawn.cpp
+++ b/Templates/TP_VehicleAdv/Source/TP_VehicleAdv/TP_VehicleAdvPawn.cpp
@@ -313,19 +313,25 @@ void ATP_VehicleAdvPawn::UpdateHUDStrings()
 {
 	float KPH = FMath::Abs(GetVehicleMovement()->GetForwardSpeed()) * 0.036f;
 	int32 KPH_int = FMath::FloorToInt(KPH);
+	int32 Gear = GetVehicleMovement()->GetCurrentGear();
 
 	// Using FText because this is display text that should be localizable
 	SpeedDisplayString = FText::Format(LOCTEXT("SpeedFormat", "{0} km/h"), FText::AsNumber(KPH_int));
-	
+
+
 	if (bInReverseGear == true)
 	{
 		GearDisplayString = FText(LOCTEXT("ReverseGear", "R"));
 	}
-	else
+	else if (Gear == 0)
 	{
-		int32 Gear = GetVehicleMovement()->GetCurrentGear();
 		GearDisplayString = (Gear == 0) ? LOCTEXT("N", "N") : FText::AsNumber(Gear);
-	}	
+	}
+	else
+	{
+		GearDisplayString = (KPH_int == 0) ? LOCTEXT("N", "N") : FText::AsNumber(Gear);
+	}
+
 }
 
 void ATP_VehicleAdvPawn::SetupInCarHUD()
