Changing vehicle advanced gears back to default.

diff --git a/Templates/TP_VehicleAdv/Source/TP_VehicleAdv/TP_VehicleAdvPawn.cpp b/Templates/TP_VehicleAdv/Source/TP_VehicleAdv/TP_VehicleAdvPawn.cpp
--- a/Templates/TP_VehicleAdv/Source/TP_VehicleAdv/TP_VehicleAdvPawn.cpp
+++ b/Templates/TP_VehicleAdv/Source/TP_VehicleAdv/TP_VehicleAdvPawn.cpp
@@ -323,13 +323,9 @@ void ATP_VehicleAdvPawn::UpdateHUDStrings()
 	{
 		GearDisplayString = FText(LOCTEXT("ReverseGear", "R"));
 	}
-	else if (Gear == 0)
-	{
-		GearDisplayString = (Gear == 0) ? LOCTEXT("N", "N") : FText::AsNumber(Gear);
-	}
 	else
 	{
-		GearDisplayString = (KPH_int == 0) ? LOCTEXT("N", "N") : FText::AsNumber(Gear);
+		GearDisplayString = (Gear == 0) ? LOCTEXT("N", "N") : FText::AsNumber(Gear);
 	}
 
 }
