Ensure that HMD_MODULE_INCLUDED is treated as 0 to avoid compile errors about missing define.

diff --git a/Templates/TP_Vehicle/Source/TP_Vehicle/TP_VehicleHud.cpp b/Templates/TP_Vehicle/Source/TP_Vehicle/TP_VehicleHud.cpp
--- a/Templates/TP_Vehicle/Source/TP_Vehicle/TP_VehicleHud.cpp
+++ b/Templates/TP_Vehicle/Source/TP_Vehicle/TP_VehicleHud.cpp
@@ -14,6 +14,10 @@
 
 #define LOCTEXT_NAMESPACE "VehicleHUD"
 
+#ifndef HMD_MODULE_INCLUDED
+#define HMD_MODULE_INCLUDED 0
+#endif
+
 ATP_VehicleHud::ATP_VehicleHud()
 {
 	static ConstructorHelpers::FObjectFinder<UFont> Font(TEXT("/Engine/EngineFonts/RobotoDistanceField"));
diff --git a/Templates/TP_Vehicle/Source/TP_Vehicle/TP_VehiclePawn.cpp b/Templates/TP_Vehicle/Source/TP_Vehicle/TP_VehiclePawn.cpp
--- a/Templates/TP_Vehicle/Source/TP_Vehicle/TP_VehiclePawn.cpp
+++ b/Templates/TP_Vehicle/Source/TP_Vehicle/TP_VehiclePawn.cpp
@@ -16,6 +16,10 @@
 #include "Materials/Material.h"
 #include "GameFramework/Controller.h"
 
+#ifndef HMD_MODULE_INCLUDED
+#define HMD_MODULE_INCLUDED 0
+#endif
+
 // Needed for VR Headset
 #if HMD_MODULE_INCLUDED
 #include "IXRTrackingSystem.h"
diff --git a/Templates/TP_VehicleAdv/Source/TP_VehicleAdv/TP_VehicleAdvHud.cpp b/Templates/TP_VehicleAdv/Source/TP_VehicleAdv/TP_VehicleAdvHud.cpp
--- a/Templates/TP_VehicleAdv/Source/TP_VehicleAdv/TP_VehicleAdvHud.cpp
+++ b/Templates/TP_VehicleAdv/Source/TP_VehicleAdv/TP_VehicleAdvHud.cpp
@@ -13,6 +13,10 @@
 
 #define LOCTEXT_NAMESPACE "VehicleHUD"
 
+#ifndef HMD_MODULE_INCLUDED
+#define HMD_MODULE_INCLUDED 0
+#endif
+
 ATP_VehicleAdvHud::ATP_VehicleAdvHud()
 {
 	static ConstructorHelpers::FObjectFinder<UFont> Font(TEXT("/Engine/EngineFonts/RobotoDistanceField"));
diff --git a/Templates/TP_VehicleAdv/Source/TP_VehicleAdv/TP_VehicleAdvPawn.cpp b/Templates/TP_VehicleAdv/Source/TP_VehicleAdv/TP_VehicleAdvPawn.cpp
--- a/Templates/TP_VehicleAdv/Source/TP_VehicleAdv/TP_VehicleAdvPawn.cpp
+++ b/Templates/TP_VehicleAdv/Source/TP_VehicleAdv/TP_VehicleAdvPawn.cpp
@@ -18,6 +18,10 @@
 #include "GameFramework/Controller.h"
 #include "UObject/ConstructorHelpers.h"
 
+#ifndef HMD_MODULE_INCLUDED
+#define HMD_MODULE_INCLUDED 0
+#endif
+
 // Needed for VR Headset
 #if HMD_MODULE_INCLUDED
 #include "IXRTrackingSystem.h"
