Fix missing include causing issues on some platforms

diff --git a/Templates/TP_VehicleAdv/Source/TP_VehicleAdv/TP_VehicleAdvPawn.cpp b/Templates/TP_VehicleAdv/Source/TP_VehicleAdv/TP_VehicleAdvPawn.cpp
--- a/Templates/TP_VehicleAdv/Source/TP_VehicleAdv/TP_VehicleAdvPawn.cpp
+++ b/Templates/TP_VehicleAdv/Source/TP_VehicleAdv/TP_VehicleAdvPawn.cpp
@@ -17,6 +17,7 @@
 #include "Engine/Engine.h"
 #include "GameFramework/Controller.h"
 #include "UObject/ConstructorHelpers.h"
+#include "GameFramework/PlayerController.h"
 
 #ifndef HMD_MODULE_INCLUDED
 #define HMD_MODULE_INCLUDED 0
