VehicleTemplate : fixing template c++ code so there are no loading errors and the car and camera have similar settings to the Vehicle BP Template
Also disabling motion blur in the level and removing substepping as this doesn't work with async mode.

diff --git a/Templates/TP_VehicleAdv/Config/DefaultEngine.ini b/Templates/TP_VehicleAdv/Config/DefaultEngine.ini
--- a/Templates/TP_VehicleAdv/Config/DefaultEngine.ini
+++ b/Templates/TP_VehicleAdv/Config/DefaultEngine.ini
@@ -1,5 +1,6 @@
 [/Script/Engine.PhysicsSettings]
-bSubstepping=True
+bSubstepping=False
+bSubsteppingAsync=False
 
 [/Script/EngineSettings.GameMapsSettings]
 EditorStartupMap=/Game/VehicleTemplate/Maps/VehicleAdvExampleMap.VehicleAdvExampleMap
diff --git a/Templates/TP_VehicleAdv/Source/TP_VehicleAdv/TP_VehicleAdvPawn.cpp b/Templates/TP_VehicleAdv/Source/TP_VehicleAdv/TP_VehicleAdvPawn.cpp
--- a/Templates/TP_VehicleAdv/Source/TP_VehicleAdv/TP_VehicleAdvPawn.cpp
+++ b/Templates/TP_VehicleAdv/Source/TP_VehicleAdv/TP_VehicleAdvPawn.cpp
@@ -40,11 +40,11 @@ PRAGMA_DISABLE_DEPRECATION_WARNINGS
 ATP_VehicleAdvPawn::ATP_VehicleAdvPawn()
 {
 	// Car mesh
-	static ConstructorHelpers::FObjectFinder<USkeletalMesh> CarMesh(TEXT("/Game/Vehicles/Vehicle/Vehicle_SkelMesh.Vehicle_SkelMesh"));
+	static ConstructorHelpers::FObjectFinder<USkeletalMesh> CarMesh(TEXT("/Game/Vehicles/SportsCar/SKM_SportsCar.SKM_SportsCar"));
 	GetMesh()->SetSkeletalMesh(CarMesh.Object);
 	GetMesh()->SetSimulatePhysics(true);
 	
-	static ConstructorHelpers::FClassFinder<UObject> AnimBPClass(TEXT("/Game/Vehicles/Vehicle/VehicleAnimationBlueprint"));
+	static ConstructorHelpers::FClassFinder<UObject> AnimBPClass(TEXT("/Game/Vehicles/SportsCar/SportsCar_AnimBP"));
 	GetMesh()->SetAnimationMode(EAnimationMode::AnimationBlueprint);
 	GetMesh()->SetAnimInstanceClass(AnimBPClass.Class);
 
@@ -56,39 +56,41 @@ ATP_VehicleAdvPawn::ATP_VehicleAdvPawn()
 	NonSlipperyMaterial = NonSlipperyMat.Object;
 
 	UChaosWheeledVehicleMovementComponent* VehicleMovement = CastChecked<UChaosWheeledVehicleMovementComponent>(GetVehicleMovement());
+	VehicleMovement->bLegacyWheelFrictionPosition = true;
 
 	// Wheels/Tyres
 	// Setup the wheels
 	VehicleMovement->WheelSetups.SetNum(4);
 	{
 		VehicleMovement->WheelSetups[0].WheelClass = UTP_VehicleAdvWheelFront::StaticClass();
-		VehicleMovement->WheelSetups[0].BoneName = FName("PhysWheel_FL");
+		VehicleMovement->WheelSetups[0].BoneName = FName("Phys_Wheel_FL");
 		VehicleMovement->WheelSetups[0].AdditionalOffset = FVector(0.f, -8.f, 0.f);
 
 		VehicleMovement->WheelSetups[1].WheelClass = UTP_VehicleAdvWheelFront::StaticClass();
-		VehicleMovement->WheelSetups[1].BoneName = FName("PhysWheel_FR");
+		VehicleMovement->WheelSetups[1].BoneName = FName("Phys_Wheel_FR");
 		VehicleMovement->WheelSetups[1].AdditionalOffset = FVector(0.f, 8.f, 0.f);
 
 		VehicleMovement->WheelSetups[2].WheelClass = UTP_VehicleAdvWheelRear::StaticClass();
-		VehicleMovement->WheelSetups[2].BoneName = FName("PhysWheel_BL");
+		VehicleMovement->WheelSetups[2].BoneName = FName("Phys_Wheel_BL");
 		VehicleMovement->WheelSetups[2].AdditionalOffset = FVector(0.f, -8.f, 0.f);
 
 		VehicleMovement->WheelSetups[3].WheelClass = UTP_VehicleAdvWheelRear::StaticClass();
-		VehicleMovement->WheelSetups[3].BoneName = FName("PhysWheel_BR");
+		VehicleMovement->WheelSetups[3].BoneName = FName("Phys_Wheel_BR");
 		VehicleMovement->WheelSetups[3].AdditionalOffset = FVector(0.f, 8.f, 0.f);
 	}
 
 	// Engine 
 	// Torque setup
-	VehicleMovement->EngineSetup.MaxRPM = 5700.0f;
-	VehicleMovement->EngineSetup.MaxTorque = 500.0f;
+	VehicleMovement->EngineSetup.MaxRPM = 7000.0f;
+	VehicleMovement->EngineSetup.MaxTorque = 750.0f;
+	VehicleMovement->EngineSetup.EngineIdleRPM = 900.0f;
 	VehicleMovement->EngineSetup.TorqueCurve.GetRichCurve()->Reset();
 	VehicleMovement->EngineSetup.TorqueCurve.GetRichCurve()->AddKey(0.0f, 400.0f);
 	VehicleMovement->EngineSetup.TorqueCurve.GetRichCurve()->AddKey(1890.0f, 500.0f);
 	VehicleMovement->EngineSetup.TorqueCurve.GetRichCurve()->AddKey(5730.0f, 400.0f);
  
 	// This works because the AxleType has been setup on the wheels
-	VehicleMovement->DifferentialSetup.DifferentialType = EVehicleDifferential::AllWheelDrive;
+	VehicleMovement->DifferentialSetup.DifferentialType = EVehicleDifferential::RearWheelDrive;
 
 	// Adjust the steering 
 	VehicleMovement->SteeringSetup.SteeringCurve.GetRichCurve()->Reset();
@@ -96,13 +98,26 @@ ATP_VehicleAdvPawn::ATP_VehicleAdvPawn()
 	VehicleMovement->SteeringSetup.SteeringCurve.GetRichCurve()->AddKey(40.0f, 0.7f);
 	VehicleMovement->SteeringSetup.SteeringCurve.GetRichCurve()->AddKey(120.0f, 0.6f);
 			
-	// Drive the front wheels a little more than the rear
-	VehicleMovement->DifferentialSetup.FrontRearSplit = 0.65;
-
 	// Automatic gearbox
 	VehicleMovement->TransmissionSetup.bUseAutomaticGears = true;
 	VehicleMovement->TransmissionSetup.bUseAutoReverse = true;
-	VehicleMovement->TransmissionSetup.GearChangeTime = 0.15f;
+	VehicleMovement->TransmissionSetup.GearChangeTime = 0.2f;
+
+	VehicleMovement->TransmissionSetup.ForwardGearRatios.Reset();
+	VehicleMovement->TransmissionSetup.ForwardGearRatios.Add(4.25f);
+	VehicleMovement->TransmissionSetup.ForwardGearRatios.Add(2.52f);
+	VehicleMovement->TransmissionSetup.ForwardGearRatios.Add(1.66f);
+	VehicleMovement->TransmissionSetup.ForwardGearRatios.Add(1.22f);
+	VehicleMovement->TransmissionSetup.ForwardGearRatios.Add(1.0f);
+
+	VehicleMovement->TransmissionSetup.ReverseGearRatios.Reset();
+	VehicleMovement->TransmissionSetup.ReverseGearRatios.Add(4.04f);
+
+	VehicleMovement->TransmissionSetup.FinalRatio = 2.81f;
+
+	VehicleMovement->TransmissionSetup.ChangeUpRPM = 6000;
+	VehicleMovement->TransmissionSetup.ChangeDownRPM = 2000;
+
 
 	// Physics settings
 	// Adjust the center of mass - the buggy is quite low
@@ -117,26 +132,29 @@ ATP_VehicleAdvPawn::ATP_VehicleAdvPawn()
 
 	// Create a spring arm component for our chase camera
 	SpringArm = CreateDefaultSubobject<USpringArmComponent>(TEXT("SpringArm"));
-	SpringArm->SetRelativeLocation(FVector(0.0f, 0.0f, 34.0f));
-	SpringArm->SetWorldRotation(FRotator(-20.0f, 0.0f, 0.0f));
+	SpringArm->SetRelativeLocation(FVector(-160.0f, 0.0f, 75.0f));
+	SpringArm->SetWorldRotation(FRotator(0.0f, 0.0f, 0.0f));
+	SpringArm->SocketOffset = FVector(0.0f, 0.0f, 150.0f);
 	SpringArm->SetupAttachment(RootComponent);
-	SpringArm->TargetArmLength = 125.0f;
+	SpringArm->TargetArmLength = 600;
 	SpringArm->bEnableCameraLag = false;
-	SpringArm->bEnableCameraRotationLag = false;
-	SpringArm->bInheritPitch = true;
+	SpringArm->bEnableCameraRotationLag = true;
+	SpringArm->CameraRotationLagSpeed = 2.0f;
+	SpringArm->bInheritPitch = false;
 	SpringArm->bInheritYaw = true;
-	SpringArm->bInheritRoll = true;
+	SpringArm->bInheritRoll = false;
 
 	// Create the chase camera component 
 	Camera = CreateDefaultSubobject<UCameraComponent>(TEXT("ChaseCamera"));
 	Camera->SetupAttachment(SpringArm, USpringArmComponent::SocketName);
-	Camera->SetRelativeLocation(FVector(-125.0, 0.0f, 0.0f));
-	Camera->SetRelativeRotation(FRotator(10.0f, 0.0f, 0.0f));
+	Camera->SetRelativeLocation(FVector(0.0f, 0.0f, 0.0f));
+	Camera->SetRelativeRotation(FRotator(0.0f, 0.0f, 0.0f));
 	Camera->bUsePawnControlRotation = false;
+	Camera->bUseFieldOfViewForLOD = true;
 	Camera->FieldOfView = 90.f;
 
 	// Create In-Car camera component 
-	InternalCameraOrigin = FVector(-34.0f, -10.0f, 50.0f);
+	InternalCameraOrigin = FVector(50.0f, 0.0f, 120.0f);
 	InternalCameraBase = CreateDefaultSubobject<USceneComponent>(TEXT("InternalCameraBase"));
 	InternalCameraBase->SetRelativeLocation(InternalCameraOrigin);
 	InternalCameraBase->SetupAttachment(GetMesh());
@@ -162,10 +180,10 @@ ATP_VehicleAdvPawn::ATP_VehicleAdvPawn()
 	InCarGear->SetupAttachment(GetMesh());
 	
 	// Setup the audio component and allocate it a sound cue
-	static ConstructorHelpers::FObjectFinder<USoundCue> SoundCue(TEXT("/Game/Vehicles/Sound/Engine_Loop_Cue.Engine_Loop_Cue"));
-	EngineSoundComponent = CreateDefaultSubobject<UAudioComponent>(TEXT("EngineSound"));
-	EngineSoundComponent->SetSound(SoundCue.Object);
-	EngineSoundComponent->SetupAttachment(GetMesh());
+	//static ConstructorHelpers::FObjectFinder<USoundCue> SoundCue(TEXT("/Game/Vehicles/Sound/Engine_Loop_Cue.Engine_Loop_Cue"));
+	//EngineSoundComponent = CreateDefaultSubobject<UAudioComponent>(TEXT("EngineSound"));
+	//EngineSoundComponent->SetSound(SoundCue.Object);
+	//EngineSoundComponent->SetupAttachment(GetMesh());
 
 	// Colors for the in-car gear display. One for normal one for reverse
 	GearDisplayReverseColor = FColor(255, 0, 0, 255);
@@ -196,7 +214,7 @@ void ATP_VehicleAdvPawn::SetupPlayerInputComponent(class UInputComponent* Player
 
 void ATP_VehicleAdvPawn::MoveForward(float Val)
 {
-	if (Val > 0)
+	if (Val >= 0)
 	{
 		GetVehicleMovementComponent()->SetThrottleInput(Val);
 		GetVehicleMovementComponent()->SetBrakeInput(0.f);
@@ -288,7 +306,7 @@ void ATP_VehicleAdvPawn::Tick(float Delta)
 	// Pass the engine RPM to the sound component
 	UChaosWheeledVehicleMovementComponent* WheeledVehicle = static_cast<UChaosWheeledVehicleMovementComponent*>(GetVehicleMovement());
 	float RPMToAudioScale = 2500.0f / WheeledVehicle->GetEngineMaxRotationSpeed();
-	EngineSoundComponent->SetFloatParameter(EngineAudioRPM, WheeledVehicle->GetEngineRotationSpeed()*RPMToAudioScale);
+	//EngineSoundComponent->SetFloatParameter(EngineAudioRPM, WheeledVehicle->GetEngineRotationSpeed()*RPMToAudioScale);
 }
 
 void ATP_VehicleAdvPawn::BeginPlay()
@@ -296,6 +314,7 @@ void ATP_VehicleAdvPawn::BeginPlay()
 	Super::BeginPlay();
 
 	bool bWantInCar = false;
+
 	// First disable both speed/gear displays 
 	bInCarCameraActive = false;
 	InCarSpeed->SetVisibility(bInCarCameraActive);
@@ -307,8 +326,9 @@ void ATP_VehicleAdvPawn::BeginPlay()
 #endif // HMD_MODULE_INCLUDED
 
 	EnableIncarView(bWantInCar);
+
 	// Start an engine sound playing
-	EngineSoundComponent->Play();
+	//EngineSoundComponent->Play();
 }
 
 void ATP_VehicleAdvPawn::OnResetVR()
diff --git a/Templates/TP_VehicleAdv/Source/TP_VehicleAdv/TP_VehicleAdvWheelFront.cpp b/Templates/TP_VehicleAdv/Source/TP_VehicleAdv/TP_VehicleAdvWheelFront.cpp
--- a/Templates/TP_VehicleAdv/Source/TP_VehicleAdv/TP_VehicleAdvWheelFront.cpp
+++ b/Templates/TP_VehicleAdv/Source/TP_VehicleAdv/TP_VehicleAdvWheelFront.cpp
@@ -7,21 +7,21 @@ PRAGMA_DISABLE_DEPRECATION_WARNINGS
 
 UTP_VehicleAdvWheelFront::UTP_VehicleAdvWheelFront()
 {
-	WheelRadius = 18.f;
-	WheelWidth = 20.0f;
-	FrictionForceMultiplier = 2.0f;
-	CorneringStiffness = 500.0f;
+	WheelRadius = 39.f;
+	WheelWidth = 35.0f;
+	FrictionForceMultiplier = 3.0f;
+	CorneringStiffness = 1200.0f;
 	bAffectedByEngine = false;
 	bAffectedByHandbrake = false;
 	bAffectedBySteering = true;
 	AxleType = EAxleType::Front;
-	SpringRate = 200.0f;
-	SpringPreload = 100.f;
+	SpringRate = 250.0f;
+	SpringPreload = 50.f;
 	SuspensionDampingRatio = 0.5f;
-	WheelLoadRatio = 0.0f;
-	RollbarScaling = 0.5f;
-	SuspensionMaxRaise = 8;
-	SuspensionMaxDrop = 12;
+	WheelLoadRatio = 0.5f;
+	RollbarScaling = 0.15f;
+	SuspensionMaxRaise = 10.0f;
+	SuspensionMaxDrop = 10.0f;
 	WheelLoadRatio = 0.5f;
 	MaxSteerAngle = 40.f;
 }
diff --git a/Templates/TP_VehicleAdv/Source/TP_VehicleAdv/TP_VehicleAdvWheelRear.cpp b/Templates/TP_VehicleAdv/Source/TP_VehicleAdv/TP_VehicleAdvWheelRear.cpp
--- a/Templates/TP_VehicleAdv/Source/TP_VehicleAdv/TP_VehicleAdvWheelRear.cpp
+++ b/Templates/TP_VehicleAdv/Source/TP_VehicleAdv/TP_VehicleAdvWheelRear.cpp
@@ -7,21 +7,21 @@ PRAGMA_DISABLE_DEPRECATION_WARNINGS
 
 UTP_VehicleAdvWheelRear::UTP_VehicleAdvWheelRear()
 {
-	WheelRadius = 18.f;
-	WheelWidth = 20.0f;
-	FrictionForceMultiplier = 2.0f;
-	CorneringStiffness = 500.0f;
+	WheelRadius = 40.f;
+	WheelWidth = 40.0f;
+	FrictionForceMultiplier = 4.0f;
+	CorneringStiffness = 1200.0f;
 	bAffectedByEngine = true;
 	bAffectedByHandbrake = true;
 	bAffectedBySteering = false;
 	AxleType = EAxleType::Rear;
-	SpringRate = 200.0f;
-	SpringPreload = 100.f;
+	SpringRate = 250.0f;
+	SpringPreload = 50.f;
 	SuspensionDampingRatio = 0.5f;
-	WheelLoadRatio = 0.0f;
+	WheelLoadRatio = 0.5f;
 	RollbarScaling = 0.5f;
-	SuspensionMaxRaise = 8;
-	SuspensionMaxDrop = 12.0f;
+	SuspensionMaxRaise = 10.0f;
+	SuspensionMaxDrop = 10.0f;
 	WheelLoadRatio = 0.5f;
 }
 
