Updated C++ Vehicle template to match its BP counterpart.
Moved from UE5-Main to Release-5.3

diff --git a/Templates/TP_VehicleAdv/Source/TP_VehicleAdv/TP_VehicleAdv.Build.cs b/Templates/TP_VehicleAdv/Source/TP_VehicleAdv/TP_VehicleAdv.Build.cs
--- a/Templates/TP_VehicleAdv/Source/TP_VehicleAdv/TP_VehicleAdv.Build.cs
+++ b/Templates/TP_VehicleAdv/Source/TP_VehicleAdv/TP_VehicleAdv.Build.cs
@@ -8,6 +8,6 @@ public class TP_VehicleAdv : ModuleRules
 	{
 		PCHUsage = PCHUsageMode.UseExplicitOrSharedPCHs;
 
-		PublicDependencyModuleNames.AddRange(new string[] { "Core", "CoreUObject", "Engine", "InputCore", "ChaosVehicles", "PhysicsCore" });
+		PublicDependencyModuleNames.AddRange(new string[] { "Core", "CoreUObject", "Engine", "InputCore", "EnhancedInput", "ChaosVehicles", "PhysicsCore" });
 	}
 }
diff --git a/Templates/TP_VehicleAdv/Source/TP_VehicleAdv/TP_VehicleAdvGameMode.cpp b/Templates/TP_VehicleAdv/Source/TP_VehicleAdv/TP_VehicleAdvGameMode.cpp
--- a/Templates/TP_VehicleAdv/Source/TP_VehicleAdv/TP_VehicleAdvGameMode.cpp
+++ b/Templates/TP_VehicleAdv/Source/TP_VehicleAdv/TP_VehicleAdvGameMode.cpp
@@ -1,11 +1,9 @@
 // Copyright Epic Games, Inc. All Rights Reserved.
 
 #include "TP_VehicleAdvGameMode.h"
-#include "TP_VehicleAdvPawn.h"
-#include "TP_VehicleAdvHud.h"
+#include "TP_VehicleAdvPlayerController.h"
 
 ATP_VehicleAdvGameMode::ATP_VehicleAdvGameMode()
 {
-	DefaultPawnClass = ATP_VehicleAdvPawn::StaticClass();
-	HUDClass = ATP_VehicleAdvHud::StaticClass();
+	PlayerControllerClass = ATP_VehicleAdvPlayerController::StaticClass();
 }
diff --git a/Templates/TP_VehicleAdv/Source/TP_VehicleAdv/TP_VehicleAdvHud.cpp b/Templates/TP_VehicleAdv/Source/TP_VehicleAdv/TP_VehicleAdvHud.cpp
deleted file mode 100644
diff --git a/Templates/TP_VehicleAdv/Source/TP_VehicleAdv/TP_VehicleAdvHud.h b/Templates/TP_VehicleAdv/Source/TP_VehicleAdv/TP_VehicleAdvHud.h
deleted file mode 100644
diff --git a/Templates/TP_VehicleAdv/Source/TP_VehicleAdv/TP_VehicleAdvOffroadCar.cpp b/Templates/TP_VehicleAdv/Source/TP_VehicleAdv/TP_VehicleAdvOffroadCar.cpp
new file mode 100644
--- /dev/null
+++ b/Templates/TP_VehicleAdv/Source/TP_VehicleAdv/TP_VehicleAdvOffroadCar.cpp
@@ -0,0 +1,85 @@
+// Copyright Epic Games, Inc. All Rights Reserved.
+
+
+#include "TP_VehicleAdvOffroadCar.h"
+#include "TP_VehicleAdvOffroadWheelFront.h"
+#include "TP_VehicleAdvOffroadWheelRear.h"
+#include "ChaosWheeledVehicleMovementComponent.h"
+#include "GameFramework/SpringArmComponent.h"
+
+ATP_VehicleAdvOffroadCar::ATP_VehicleAdvOffroadCar()
+{
+	// construct the mesh components
+	Chassis = CreateDefaultSubobject<UStaticMeshComponent>(TEXT("Chassis"));
+	Chassis->SetupAttachment(GetMesh());
+
+	// NOTE: tire sockets are set from the Blueprint class
+	TireFrontLeft = CreateDefaultSubobject<UStaticMeshComponent>(TEXT("Tire Front Left"));
+	TireFrontLeft->SetupAttachment(GetMesh(), FName("VisWheel_FL"));
+	TireFrontLeft->SetCollisionProfileName(FName("NoCollision"));
+
+	TireFrontRight = CreateDefaultSubobject<UStaticMeshComponent>(TEXT("Tire Front Right"));
+	TireFrontRight->SetupAttachment(GetMesh(), FName("VisWheel_FR"));
+	TireFrontRight->SetCollisionProfileName(FName("NoCollision"));
+	TireFrontRight->SetRelativeRotation(FRotator(0.0f, 180.0f, 0.0f));
+
+	TireRearLeft = CreateDefaultSubobject<UStaticMeshComponent>(TEXT("Tire Rear Left"));
+	TireRearLeft->SetupAttachment(GetMesh(), FName("VisWheel_BL"));
+	TireRearLeft->SetCollisionProfileName(FName("NoCollision"));
+
+	TireRearRight = CreateDefaultSubobject<UStaticMeshComponent>(TEXT("Tire Rear Right"));
+	TireRearRight->SetupAttachment(GetMesh(), FName("VisWheel_BR"));
+	TireRearRight->SetCollisionProfileName(FName("NoCollision"));
+	TireRearRight->SetRelativeRotation(FRotator(0.0f, 180.0f, 0.0f));
+
+	// adjust the cameras
+	GetFrontSpringArm()->SetRelativeLocation(FVector(-5.0f, -30.0f, 135.0f));
+	GetBackSpringArm()->SetRelativeLocation(FVector(0.0f, 0.0f, 75.0f));
+
+	// Note: for faster iteration times, the vehicle setup can be tweaked in the Blueprint instead
+
+	// Set up the chassis
+	GetChaosVehicleMovement()->ChassisHeight = 160.0f;
+	GetChaosVehicleMovement()->DragCoefficient = 0.1f;
+	GetChaosVehicleMovement()->DownforceCoefficient = 0.1f;
+	GetChaosVehicleMovement()->CenterOfMassOverride = FVector(0.0f, 0.0f, 75.0f);
+	GetChaosVehicleMovement()->bEnableCenterOfMassOverride = true;
+
+	// Set up the wheels
+	GetChaosVehicleMovement()->bLegacyWheelFrictionPosition = true;
+	GetChaosVehicleMovement()->WheelSetups.SetNum(4);
+
+	GetChaosVehicleMovement()->WheelSetups[0].WheelClass = UTP_VehicleAdvOffroadWheelFront::StaticClass();
+	GetChaosVehicleMovement()->WheelSetups[0].BoneName = FName("PhysWheel_FL");
+	GetChaosVehicleMovement()->WheelSetups[0].AdditionalOffset = FVector(0.0f, 0.0f, 0.0f);
+
+	GetChaosVehicleMovement()->WheelSetups[1].WheelClass = UTP_VehicleAdvOffroadWheelFront::StaticClass();
+	GetChaosVehicleMovement()->WheelSetups[1].BoneName = FName("PhysWheel_FR");
+	GetChaosVehicleMovement()->WheelSetups[1].AdditionalOffset = FVector(0.0f, 0.0f, 0.0f);
+
+	GetChaosVehicleMovement()->WheelSetups[2].WheelClass = UTP_VehicleAdvOffroadWheelRear::StaticClass();
+	GetChaosVehicleMovement()->WheelSetups[2].BoneName = FName("PhysWheel_BL");
+	GetChaosVehicleMovement()->WheelSetups[2].AdditionalOffset = FVector(0.0f, 0.0f, 0.0f);
+
+	GetChaosVehicleMovement()->WheelSetups[3].WheelClass = UTP_VehicleAdvOffroadWheelRear::StaticClass();
+	GetChaosVehicleMovement()->WheelSetups[3].BoneName = FName("PhysWheel_BR");
+	GetChaosVehicleMovement()->WheelSetups[3].AdditionalOffset = FVector(0.0f, 0.0f, 0.0f);
+
+	// Set up the engine
+	// NOTE: Check the Blueprint asset for the Torque Curve
+	GetChaosVehicleMovement()->EngineSetup.MaxTorque = 600.0f;
+	GetChaosVehicleMovement()->EngineSetup.MaxRPM = 5000.0f;
+	GetChaosVehicleMovement()->EngineSetup.EngineIdleRPM = 1200.0f;
+	GetChaosVehicleMovement()->EngineSetup.EngineBrakeEffect = 0.05f;
+	GetChaosVehicleMovement()->EngineSetup.EngineRevUpMOI = 5.0f;
+	GetChaosVehicleMovement()->EngineSetup.EngineRevDownRate = 600.0f;
+
+	// Set up the differential
+	GetChaosVehicleMovement()->DifferentialSetup.DifferentialType = EVehicleDifferential::AllWheelDrive;
+	GetChaosVehicleMovement()->DifferentialSetup.FrontRearSplit = 0.5f;
+
+	// Set up the steering
+	// NOTE: Check the Blueprint asset for the Steering Curve
+	GetChaosVehicleMovement()->SteeringSetup.SteeringType = ESteeringType::AngleRatio;
+	GetChaosVehicleMovement()->SteeringSetup.AngleRatio = 0.7f;
+}
\ No newline at end of file
diff --git a/Templates/TP_VehicleAdv/Source/TP_VehicleAdv/TP_VehicleAdvOffroadCar.h b/Templates/TP_VehicleAdv/Source/TP_VehicleAdv/TP_VehicleAdvOffroadCar.h
new file mode 100644
--- /dev/null
+++ b/Templates/TP_VehicleAdv/Source/TP_VehicleAdv/TP_VehicleAdvOffroadCar.h
@@ -0,0 +1,40 @@
+// Copyright Epic Games, Inc. All Rights Reserved.
+
+#pragma once
+
+#include "CoreMinimal.h"
+#include "TP_VehicleAdvPawn.h"
+#include "TP_VehicleAdvOffroadCar.generated.h"
+
+/**
+ *  Offroad car wheeled vehicle implementation
+ */
+UCLASS(abstract)
+class TP_VEHICLEADV_API ATP_VehicleAdvOffroadCar : public ATP_VehicleAdvPawn
+{
+	GENERATED_BODY()
+	
+	/** Chassis static mesh */
+	UPROPERTY(VisibleAnywhere, BlueprintReadOnly, Category = Meshes, meta = (AllowPrivateAccess = "true"))
+	UStaticMeshComponent* Chassis;
+
+	/** FL Tire static mesh */
+	UPROPERTY(VisibleAnywhere, BlueprintReadOnly, Category = Meshes, meta = (AllowPrivateAccess = "true"))
+	UStaticMeshComponent* TireFrontLeft;
+
+	/** FR Tire static mesh */
+	UPROPERTY(VisibleAnywhere, BlueprintReadOnly, Category = Meshes, meta = (AllowPrivateAccess = "true"))
+	UStaticMeshComponent* TireFrontRight;
+
+	/** RL Tire static mesh */
+	UPROPERTY(VisibleAnywhere, BlueprintReadOnly, Category = Meshes, meta = (AllowPrivateAccess = "true"))
+	UStaticMeshComponent* TireRearLeft;
+
+	/** RR Tire static mesh */
+	UPROPERTY(VisibleAnywhere, BlueprintReadOnly, Category = Meshes, meta = (AllowPrivateAccess = "true"))
+	UStaticMeshComponent* TireRearRight;
+
+public:
+
+	ATP_VehicleAdvOffroadCar();
+};
diff --git a/Templates/TP_VehicleAdv/Source/TP_VehicleAdv/TP_VehicleAdvOffroadWheelFront.cpp b/Templates/TP_VehicleAdv/Source/TP_VehicleAdv/TP_VehicleAdvOffroadWheelFront.cpp
new file mode 100644
--- /dev/null
+++ b/Templates/TP_VehicleAdv/Source/TP_VehicleAdv/TP_VehicleAdvOffroadWheelFront.cpp
@@ -0,0 +1,22 @@
+// Copyright Epic Games, Inc. All Rights Reserved.
+
+
+#include "TP_VehicleAdvOffroadWheelFront.h"
+
+UTP_VehicleAdvOffroadWheelFront::UTP_VehicleAdvOffroadWheelFront()
+{
+	WheelRadius = 50.0f;
+	CorneringStiffness = 750.0f;
+	FrictionForceMultiplier = 4.0f;
+	bAffectedByEngine = true;
+
+	SuspensionMaxRaise = 20.0f;
+	SuspensionMaxDrop = 20.0f;
+	WheelLoadRatio = 1.0f;
+	SpringRate = 100.0f;
+	SpringPreload = 100.0f;
+	SweepShape = ESweepShape::Shapecast;
+
+	MaxBrakeTorque = 3000.0f;
+	MaxHandBrakeTorque = 6000.0f;
+}
\ No newline at end of file
diff --git a/Templates/TP_VehicleAdv/Source/TP_VehicleAdv/TP_VehicleAdvOffroadWheelFront.h b/Templates/TP_VehicleAdv/Source/TP_VehicleAdv/TP_VehicleAdvOffroadWheelFront.h
new file mode 100644
--- /dev/null
+++ b/Templates/TP_VehicleAdv/Source/TP_VehicleAdv/TP_VehicleAdvOffroadWheelFront.h
@@ -0,0 +1,19 @@
+// Copyright Epic Games, Inc. All Rights Reserved.
+
+#pragma once
+
+#include "CoreMinimal.h"
+#include "TP_VehicleAdvWheelFront.h"
+#include "TP_VehicleAdvOffroadWheelFront.generated.h"
+
+/**
+ *  Front wheel definition for Offroad Car.
+ */
+UCLASS()
+class TP_VEHICLEADV_API UTP_VehicleAdvOffroadWheelFront : public UTP_VehicleAdvWheelFront
+{
+	GENERATED_BODY()
+	
+public:
+	UTP_VehicleAdvOffroadWheelFront();
+};
diff --git a/Templates/TP_VehicleAdv/Source/TP_VehicleAdv/TP_VehicleAdvOffroadWheelRear.cpp b/Templates/TP_VehicleAdv/Source/TP_VehicleAdv/TP_VehicleAdvOffroadWheelRear.cpp
new file mode 100644
--- /dev/null
+++ b/Templates/TP_VehicleAdv/Source/TP_VehicleAdv/TP_VehicleAdvOffroadWheelRear.cpp
@@ -0,0 +1,21 @@
+// Copyright Epic Games, Inc. All Rights Reserved.
+
+
+#include "TP_VehicleAdvOffroadWheelRear.h"
+
+UTP_VehicleAdvOffroadWheelRear::UTP_VehicleAdvOffroadWheelRear()
+{
+	WheelRadius = 50.f;
+	CorneringStiffness = 750.0f;
+	FrictionForceMultiplier = 4.0f;
+	
+	SuspensionMaxRaise = 20.0f;
+	SuspensionMaxDrop = 20.0f;
+	WheelLoadRatio = 1.0f;
+	SpringRate = 100.0f;
+	SpringPreload = 100.0f;
+	SweepShape = ESweepShape::Shapecast;
+
+	MaxBrakeTorque = 3000.0f;
+	MaxHandBrakeTorque = 6000.0f;
+}
\ No newline at end of file
diff --git a/Templates/TP_VehicleAdv/Source/TP_VehicleAdv/TP_VehicleAdvOffroadWheelRear.h b/Templates/TP_VehicleAdv/Source/TP_VehicleAdv/TP_VehicleAdvOffroadWheelRear.h
new file mode 100644
--- /dev/null
+++ b/Templates/TP_VehicleAdv/Source/TP_VehicleAdv/TP_VehicleAdvOffroadWheelRear.h
@@ -0,0 +1,20 @@
+// Copyright Epic Games, Inc. All Rights Reserved.
+
+#pragma once
+
+#include "CoreMinimal.h"
+#include "TP_VehicleAdvWheelRear.h"
+#include "TP_VehicleAdvOffroadWheelRear.generated.h"
+
+/**
+ *  Rear wheel definition for Offroad Car.
+ */
+UCLASS()
+class TP_VEHICLEADV_API UTP_VehicleAdvOffroadWheelRear : public UTP_VehicleAdvWheelRear
+{
+	GENERATED_BODY()
+	
+public:
+
+	UTP_VehicleAdvOffroadWheelRear();
+};
diff --git a/Templates/TP_VehicleAdv/Source/TP_VehicleAdv/TP_VehicleAdvPawn.cpp b/Templates/TP_VehicleAdv/Source/TP_VehicleAdv/TP_VehicleAdvPawn.cpp
--- a/Templates/TP_VehicleAdv/Source/TP_VehicleAdv/TP_VehicleAdvPawn.cpp
+++ b/Templates/TP_VehicleAdv/Source/TP_VehicleAdv/TP_VehicleAdvPawn.cpp
@@ -3,404 +3,205 @@
 #include "TP_VehicleAdvPawn.h"
 #include "TP_VehicleAdvWheelFront.h"
 #include "TP_VehicleAdvWheelRear.h"
-#include "TP_VehicleAdvHud.h"
 #include "Components/SkeletalMeshComponent.h"
 #include "GameFramework/SpringArmComponent.h"
 #include "Camera/CameraComponent.h"
-#include "Components/InputComponent.h"
-#include "Components/TextRenderComponent.h"
-#include "Components/AudioComponent.h"
-#include "Sound/SoundCue.h"
-#include "PhysicalMaterials/PhysicalMaterial.h"
+#include "EnhancedInputComponent.h"
+#include "EnhancedInputSubsystems.h"
+#include "InputActionValue.h"
 #include "ChaosWheeledVehicleMovementComponent.h"
-#include "Engine/SkeletalMesh.h"
-#include "Engine/Engine.h"
-#include "GameFramework/Controller.h"
-#include "UObject/ConstructorHelpers.h"
-#include "GameFramework/PlayerController.h"
-
-#ifndef HMD_MODULE_INCLUDED
-#define HMD_MODULE_INCLUDED 0
-#endif
-
-// Needed for VR Headset
-#if HMD_MODULE_INCLUDED
-#include "IXRTrackingSystem.h"
-#include "HeadMountedDisplayFunctionLibrary.h"
-#endif // HMD_MODULE_INCLUDED
-
-const FName ATP_VehicleAdvPawn::LookUpBinding("LookUp");
-const FName ATP_VehicleAdvPawn::LookRightBinding("LookRight");
-const FName ATP_VehicleAdvPawn::EngineAudioRPM("RPM");
 
 #define LOCTEXT_NAMESPACE "VehiclePawn"
 
-PRAGMA_DISABLE_DEPRECATION_WARNINGS
+DEFINE_LOG_CATEGORY(LogTemplateVehicle);
 
 ATP_VehicleAdvPawn::ATP_VehicleAdvPawn()
 {
-	// Car mesh
-	static ConstructorHelpers::FObjectFinder<USkeletalMesh> CarMesh(TEXT("/Game/Vehicles/SportsCar/SKM_SportsCar.SKM_SportsCar"));
-	GetMesh()->SetSkeletalMesh(CarMesh.Object);
+	// construct the front camera boom
+	FrontSpringArm = CreateDefaultSubobject<USpringArmComponent>(TEXT("Front Spring Arm"));
+	FrontSpringArm->SetupAttachment(GetMesh());
+	FrontSpringArm->TargetArmLength = 0.0f;
+	FrontSpringArm->bDoCollisionTest = false;
+	FrontSpringArm->bEnableCameraRotationLag = true;
+	FrontSpringArm->CameraRotationLagSpeed = 15.0f;
+	FrontSpringArm->SetRelativeLocation(FVector(30.0f, 0.0f, 120.0f));
+
+	FrontCamera = CreateDefaultSubobject<UCameraComponent>(TEXT("Front Camera"));
+	FrontCamera->SetupAttachment(FrontSpringArm);
+	FrontCamera->bAutoActivate = false;
+
+	// construct the back camera boom
+	BackSpringArm = CreateDefaultSubobject<USpringArmComponent>(TEXT("Back Spring Arm"));
+	BackSpringArm->SetupAttachment(GetMesh());
+	BackSpringArm->TargetArmLength = 650.0f;
+	BackSpringArm->SocketOffset.Z = 150.0f;
+	BackSpringArm->bDoCollisionTest = false;
+	BackSpringArm->bInheritPitch = false;
+	BackSpringArm->bInheritRoll = false;
+	BackSpringArm->bEnableCameraRotationLag = true;
+	BackSpringArm->CameraRotationLagSpeed = 2.0f;
+	BackSpringArm->CameraLagMaxDistance = 50.0f;
+
+	BackCamera = CreateDefaultSubobject<UCameraComponent>(TEXT("Back Camera"));
+	BackCamera->SetupAttachment(BackSpringArm);
+
+	// Configure the car mesh
 	GetMesh()->SetSimulatePhysics(true);
-	
-	static ConstructorHelpers::FClassFinder<UObject> AnimBPClass(TEXT("/Game/Vehicles/SportsCar/SportsCar_AnimBP"));
-	GetMesh()->SetAnimationMode(EAnimationMode::AnimationBlueprint);
-	GetMesh()->SetAnimInstanceClass(AnimBPClass.Class);
-
-	// Setup friction materials
-	static ConstructorHelpers::FObjectFinder<UPhysicalMaterial> SlipperyMat(TEXT("/Game/Vehicles/PhysicsMaterials/Slippery.Slippery"));
-	SlipperyMaterial = SlipperyMat.Object;
-		
-	static ConstructorHelpers::FObjectFinder<UPhysicalMaterial> NonSlipperyMat(TEXT("/Game/Vehicles/PhysicsMaterials/NonSlippery.NonSlippery"));
-	NonSlipperyMaterial = NonSlipperyMat.Object;
-
-	UChaosWheeledVehicleMovementComponent* VehicleMovement = CastChecked<UChaosWheeledVehicleMovementComponent>(GetVehicleMovement());
-	VehicleMovement->bLegacyWheelFrictionPosition = true;
-
-	// Wheels/Tyres
-	// Setup the wheels
-	VehicleMovement->WheelSetups.SetNum(4);
-	{
-		VehicleMovement->WheelSetups[0].WheelClass = UTP_VehicleAdvWheelFront::StaticClass();
-		VehicleMovement->WheelSetups[0].BoneName = FName("Phys_Wheel_FL");
-		VehicleMovement->WheelSetups[0].AdditionalOffset = FVector(0.f, -8.f, 0.f);
-
-		VehicleMovement->WheelSetups[1].WheelClass = UTP_VehicleAdvWheelFront::StaticClass();
-		VehicleMovement->WheelSetups[1].BoneName = FName("Phys_Wheel_FR");
-		VehicleMovement->WheelSetups[1].AdditionalOffset = FVector(0.f, 8.f, 0.f);
-
-		VehicleMovement->WheelSetups[2].WheelClass = UTP_VehicleAdvWheelRear::StaticClass();
-		VehicleMovement->WheelSetups[2].BoneName = FName("Phys_Wheel_BL");
-		VehicleMovement->WheelSetups[2].AdditionalOffset = FVector(0.f, -8.f, 0.f);
+	GetMesh()->SetCollisionProfileName(FName("Vehicle"));
 
-		VehicleMovement->WheelSetups[3].WheelClass = UTP_VehicleAdvWheelRear::StaticClass();
-		VehicleMovement->WheelSetups[3].BoneName = FName("Phys_Wheel_BR");
-		VehicleMovement->WheelSetups[3].AdditionalOffset = FVector(0.f, 8.f, 0.f);
-	}
-
-	// Engine 
-	// Torque setup
-	VehicleMovement->EngineSetup.MaxRPM = 7000.0f;
-	VehicleMovement->EngineSetup.MaxTorque = 750.0f;
-	VehicleMovement->EngineSetup.EngineIdleRPM = 900.0f;
-	VehicleMovement->EngineSetup.TorqueCurve.GetRichCurve()->Reset();
-	VehicleMovement->EngineSetup.TorqueCurve.GetRichCurve()->AddKey(0.0f, 400.0f);
-	VehicleMovement->EngineSetup.TorqueCurve.GetRichCurve()->AddKey(1890.0f, 500.0f);
-	VehicleMovement->EngineSetup.TorqueCurve.GetRichCurve()->AddKey(5730.0f, 400.0f);
- 
-	// This works because the AxleType has been setup on the wheels
-	VehicleMovement->DifferentialSetup.DifferentialType = EVehicleDifferential::RearWheelDrive;
-
-	// Adjust the steering 
-	VehicleMovement->SteeringSetup.SteeringCurve.GetRichCurve()->Reset();
-	VehicleMovement->SteeringSetup.SteeringCurve.GetRichCurve()->AddKey(0.0f, 1.0f);
-	VehicleMovement->SteeringSetup.SteeringCurve.GetRichCurve()->AddKey(40.0f, 0.7f);
-	VehicleMovement->SteeringSetup.SteeringCurve.GetRichCurve()->AddKey(120.0f, 0.6f);
-			
-	// Automatic gearbox
-	VehicleMovement->TransmissionSetup.bUseAutomaticGears = true;
-	VehicleMovement->TransmissionSetup.bUseAutoReverse = true;
-	VehicleMovement->TransmissionSetup.GearChangeTime = 0.2f;
-
-	VehicleMovement->TransmissionSetup.ForwardGearRatios.Reset();
-	VehicleMovement->TransmissionSetup.ForwardGearRatios.Add(4.25f);
-	VehicleMovement->TransmissionSetup.ForwardGearRatios.Add(2.52f);
-	VehicleMovement->TransmissionSetup.ForwardGearRatios.Add(1.66f);
-	VehicleMovement->TransmissionSetup.ForwardGearRatios.Add(1.22f);
-	VehicleMovement->TransmissionSetup.ForwardGearRatios.Add(1.0f);
-
-	VehicleMovement->TransmissionSetup.ReverseGearRatios.Reset();
-	VehicleMovement->TransmissionSetup.ReverseGearRatios.Add(4.04f);
-
-	VehicleMovement->TransmissionSetup.FinalRatio = 2.81f;
-
-	VehicleMovement->TransmissionSetup.ChangeUpRPM = 6000;
-	VehicleMovement->TransmissionSetup.ChangeDownRPM = 2000;
-
-
-	// Physics settings
-	// Adjust the center of mass - the buggy is quite low
-	UPrimitiveComponent* UpdatedPrimitive = Cast<UPrimitiveComponent>(VehicleMovement->UpdatedComponent);
-	if (UpdatedPrimitive)
-	{
-		UpdatedPrimitive->BodyInstance.COMNudge = FVector(8.0f, 0.0f, -15.0f);
-	}
+	// get the Chaos Wheeled movement component
+	ChaosVehicleMovement = CastChecked<UChaosWheeledVehicleMovementComponent>(GetVehicleMovement());
 
-	// Set the inertia scale. This controls how the mass of the vehicle is distributed.
-	VehicleMovement->InertiaTensorScale = FVector(1.0f, 1.333f, 1.2f);
-
-	// Create a spring arm component for our chase camera
-	SpringArm = CreateDefaultSubobject<USpringArmComponent>(TEXT("SpringArm"));
-	SpringArm->SetRelativeLocation(FVector(-160.0f, 0.0f, 75.0f));
-	SpringArm->SetWorldRotation(FRotator(0.0f, 0.0f, 0.0f));
-	SpringArm->SocketOffset = FVector(0.0f, 0.0f, 150.0f);
-	SpringArm->SetupAttachment(RootComponent);
-	SpringArm->TargetArmLength = 600;
-	SpringArm->bEnableCameraLag = false;
-	SpringArm->bEnableCameraRotationLag = true;
-	SpringArm->CameraRotationLagSpeed = 2.0f;
-	SpringArm->bInheritPitch = false;
-	SpringArm->bInheritYaw = true;
-	SpringArm->bInheritRoll = false;
-
-	// Create the chase camera component 
-	Camera = CreateDefaultSubobject<UCameraComponent>(TEXT("ChaseCamera"));
-	Camera->SetupAttachment(SpringArm, USpringArmComponent::SocketName);
-	Camera->SetRelativeLocation(FVector(0.0f, 0.0f, 0.0f));
-	Camera->SetRelativeRotation(FRotator(0.0f, 0.0f, 0.0f));
-	Camera->bUsePawnControlRotation = false;
-	Camera->bUseFieldOfViewForLOD = true;
-	Camera->FieldOfView = 90.f;
-
-	// Create In-Car camera component 
-	InternalCameraOrigin = FVector(50.0f, 0.0f, 120.0f);
-	InternalCameraBase = CreateDefaultSubobject<USceneComponent>(TEXT("InternalCameraBase"));
-	InternalCameraBase->SetRelativeLocation(InternalCameraOrigin);
-	InternalCameraBase->SetupAttachment(GetMesh());
-
-	InternalCamera = CreateDefaultSubobject<UCameraComponent>(TEXT("InternalCamera"));
-	InternalCamera->bUsePawnControlRotation = false;
-	InternalCamera->FieldOfView = 90.f;
-	InternalCamera->SetupAttachment(InternalCameraBase);
-
-	// In car HUD
-	// Create text render component for in car speed display
-	InCarSpeed = CreateDefaultSubobject<UTextRenderComponent>(TEXT("IncarSpeed"));
-	InCarSpeed->SetRelativeScale3D(FVector(0.1f, 0.1f, 0.1f));
-	InCarSpeed->SetRelativeLocation(FVector(35.0f, -6.0f, 20.0f));
-	InCarSpeed->SetRelativeRotation(FRotator(0.0f, 180.0f, 0.0f));
-	InCarSpeed->SetupAttachment(GetMesh());
-
-	// Create text render component for in car gear display
-	InCarGear = CreateDefaultSubobject<UTextRenderComponent>(TEXT("IncarGear"));
-	InCarGear->SetRelativeScale3D(FVector(0.1f, 0.1f, 0.1f));
-	InCarGear->SetRelativeLocation(FVector(35.0f, 5.0f, 20.0f));
-	InCarGear->SetRelativeRotation(FRotator(0.0f, 180.0f, 0.0f));
-	InCarGear->SetupAttachment(GetMesh());
-	
-	// Setup the audio component and allocate it a sound cue
-	//static ConstructorHelpers::FObjectFinder<USoundCue> SoundCue(TEXT("/Game/Vehicles/Sound/Engine_Loop_Cue.Engine_Loop_Cue"));
-	//EngineSoundComponent = CreateDefaultSubobject<UAudioComponent>(TEXT("EngineSound"));
-	//EngineSoundComponent->SetSound(SoundCue.Object);
-	//EngineSoundComponent->SetupAttachment(GetMesh());
-
-	// Colors for the in-car gear display. One for normal one for reverse
-	GearDisplayReverseColor = FColor(255, 0, 0, 255);
-	GearDisplayColor = FColor(255, 255, 255, 255);
-
-	bIsLowFriction = false;
-	bInReverseGear = false;
 }
 
 void ATP_VehicleAdvPawn::SetupPlayerInputComponent(class UInputComponent* PlayerInputComponent)
 {
 	Super::SetupPlayerInputComponent(PlayerInputComponent);
 
-	// set up gameplay key bindings
-	check(PlayerInputComponent);
+	if (UEnhancedInputComponent* EnhancedInputComponent = Cast<UEnhancedInputComponent>(PlayerInputComponent))
+	{
+		// steering 
+		EnhancedInputComponent->BindAction(SteeringAction, ETriggerEvent::Triggered, this, &ATP_VehicleAdvPawn::Steering);
+		EnhancedInputComponent->BindAction(SteeringAction, ETriggerEvent::Completed, this, &ATP_VehicleAdvPawn::Steering);
 
-	PlayerInputComponent->BindAxis("MoveForward", this, &ATP_VehicleAdvPawn::MoveForward);
-	PlayerInputComponent->BindAxis("MoveRight", this, &ATP_VehicleAdvPawn::MoveRight);
-	PlayerInputComponent->BindAxis(LookUpBinding);
-	PlayerInputComponent->BindAxis(LookRightBinding);
+		// throttle 
+		EnhancedInputComponent->BindAction(ThrottleAction, ETriggerEvent::Triggered, this, &ATP_VehicleAdvPawn::Throttle);
+		EnhancedInputComponent->BindAction(ThrottleAction, ETriggerEvent::Completed, this, &ATP_VehicleAdvPawn::Throttle);
 
-	PlayerInputComponent->BindAction("Handbrake", IE_Pressed, this, &ATP_VehicleAdvPawn::OnHandbrakePressed);
-	PlayerInputComponent->BindAction("Handbrake", IE_Released, this, &ATP_VehicleAdvPawn::OnHandbrakeReleased);
-	PlayerInputComponent->BindAction("SwitchCamera", IE_Pressed, this, &ATP_VehicleAdvPawn::OnToggleCamera);
+		// break 
+		EnhancedInputComponent->BindAction(BrakeAction, ETriggerEvent::Triggered, this, &ATP_VehicleAdvPawn::Brake);
+		EnhancedInputComponent->BindAction(BrakeAction, ETriggerEvent::Started, this, &ATP_VehicleAdvPawn::StartBrake);
+		EnhancedInputComponent->BindAction(BrakeAction, ETriggerEvent::Completed, this, &ATP_VehicleAdvPawn::StopBrake);
 
-	PlayerInputComponent->BindAction("ResetVR", IE_Pressed, this, &ATP_VehicleAdvPawn::OnResetVR); 
-}
+		// handbrake 
+		EnhancedInputComponent->BindAction(HandbrakeAction, ETriggerEvent::Started, this, &ATP_VehicleAdvPawn::StartHandbrake);
+		EnhancedInputComponent->BindAction(HandbrakeAction, ETriggerEvent::Completed, this, &ATP_VehicleAdvPawn::StopHandbrake);
 
-void ATP_VehicleAdvPawn::MoveForward(float Val)
-{
-	if (Val >= 0)
-	{
-		GetVehicleMovementComponent()->SetThrottleInput(Val);
-		GetVehicleMovementComponent()->SetBrakeInput(0.f);
+		// look around 
+		EnhancedInputComponent->BindAction(LookAroundAction, ETriggerEvent::Triggered, this, &ATP_VehicleAdvPawn::LookAround);
+
+		// toggle camera 
+		EnhancedInputComponent->BindAction(ToggleCameraAction, ETriggerEvent::Triggered, this, &ATP_VehicleAdvPawn::ToggleCamera);
+
+		// reset the vehicle 
+		EnhancedInputComponent->BindAction(ResetVehicleAction, ETriggerEvent::Triggered, this, &ATP_VehicleAdvPawn::ResetVehicle);
 	}
 	else
 	{
-		GetVehicleMovementComponent()->SetThrottleInput(0.f);
-		GetVehicleMovementComponent()->SetBrakeInput(-Val);
+		UE_LOG(LogTemplateVehicle, Error, TEXT("'%s' Failed to find an Enhanced Input component! This template is built to use the Enhanced Input system. If you intend to use the legacy system, then you will need to update this C++ file."), *GetNameSafe(this));
 	}
 }
 
-void ATP_VehicleAdvPawn::MoveRight(float Val)
+void ATP_VehicleAdvPawn::Tick(float Delta)
 {
-	GetVehicleMovementComponent()->SetSteeringInput(Val);
+	Super::Tick(Delta);
+
+	// add some angular damping if the vehicle is in midair
+	bool bMovingOnGround = ChaosVehicleMovement->IsMovingOnGround();
+	GetMesh()->SetAngularDamping(bMovingOnGround ? 0.0f : 3.0f);
+
+	// realign the camera yaw to face front
+	float CameraYaw = BackSpringArm->GetRelativeRotation().Yaw;
+	CameraYaw = FMath::FInterpTo(CameraYaw, 0.0f, Delta, 1.0f);
+
+	BackSpringArm->SetRelativeRotation(FRotator(0.0f, CameraYaw, 0.0f));
 }
 
-void ATP_VehicleAdvPawn::OnHandbrakePressed()
+void ATP_VehicleAdvPawn::Steering(const FInputActionValue& Value)
 {
-	GetVehicleMovementComponent()->SetHandbrakeInput(true);
+	// get the input magnitude for steering
+	float SteeringValue = Value.Get<float>();
+
+	// add the input
+	ChaosVehicleMovement->SetSteeringInput(SteeringValue);
 }
 
-void ATP_VehicleAdvPawn::OnHandbrakeReleased()
+void ATP_VehicleAdvPawn::Throttle(const FInputActionValue& Value)
 {
-	GetVehicleMovementComponent()->SetHandbrakeInput(false);
+	// get the input magnitude for the throttle
+	float ThrottleValue = Value.Get<float>();
+
+	// add the input
+	ChaosVehicleMovement->SetThrottleInput(ThrottleValue);
 }
 
-void ATP_VehicleAdvPawn::OnToggleCamera()
+void ATP_VehicleAdvPawn::Brake(const FInputActionValue& Value)
 {
-	EnableIncarView(!bInCarCameraActive);
+	// get the input magnitude for the brakes
+	float BreakValue = Value.Get<float>();
+
+	// add the input
+	ChaosVehicleMovement->SetBrakeInput(BreakValue);
 }
 
-void ATP_VehicleAdvPawn::EnableIncarView(const bool bState)
+void ATP_VehicleAdvPawn::StartBrake(const FInputActionValue& Value)
 {
-	if (bState != bInCarCameraActive)
-	{
-		bInCarCameraActive = bState;
-		
-		if (bState == true)
-		{
-			OnResetVR();
-			Camera->Deactivate();
-			InternalCamera->Activate();
-		}
-		else
-		{
-			InternalCamera->Deactivate();
-			Camera->Activate();
-		}
-		
-		InCarSpeed->SetVisibility(bInCarCameraActive);
-		InCarGear->SetVisibility(bInCarCameraActive);
-	}
+	// call the Blueprint hook for the break lights
+	BrakeLights(true);
 }
 
-void ATP_VehicleAdvPawn::Tick(float Delta)
+void ATP_VehicleAdvPawn::StopBrake(const FInputActionValue& Value)
 {
-	Super::Tick(Delta);
-
-	// Setup the flag to say we are in reverse gear
-	bInReverseGear = GetVehicleMovement()->GetCurrentGear() < 0;
-	
-	// Update phsyics material
-	UpdatePhysicsMaterial();
-
-	// Update the strings used in the hud (incar and onscreen)
-	UpdateHUDStrings();
+	// call the Blueprint hook for the break lights
+	BrakeLights(false);
 
-	// Set the string in the incar hud
-	SetupInCarHUD();
-
-	bool bHMDActive = false;
-#if HMD_MODULE_INCLUDED
-	if ((GEngine->XRSystem.IsValid() == true ) && ( (GEngine->XRSystem->IsHeadTrackingAllowed() == true) || (GEngine->IsStereoscopic3D() == true)))
-	{
-		bHMDActive = true;
-	}
-#endif // HMD_MODULE_INCLUDED
-	if( bHMDActive == false )
-	{
-		if ( (InputComponent) && (bInCarCameraActive == true ))
-		{
-			FRotator HeadRotation = InternalCamera->GetRelativeRotation();
-			HeadRotation.Pitch += InputComponent->GetAxisValue(LookUpBinding);
-			HeadRotation.Yaw += InputComponent->GetAxisValue(LookRightBinding);
-			InternalCamera->SetRelativeRotation(HeadRotation);
-		}
-	}	
-
-	// Pass the engine RPM to the sound component
-	UChaosWheeledVehicleMovementComponent* WheeledVehicle = static_cast<UChaosWheeledVehicleMovementComponent*>(GetVehicleMovement());
-	float RPMToAudioScale = 2500.0f / WheeledVehicle->GetEngineMaxRotationSpeed();
-	//EngineSoundComponent->SetFloatParameter(EngineAudioRPM, WheeledVehicle->GetEngineRotationSpeed()*RPMToAudioScale);
+	// reset brake input to zero
+	ChaosVehicleMovement->SetBrakeInput(0.0f);
 }
 
-void ATP_VehicleAdvPawn::BeginPlay()
+void ATP_VehicleAdvPawn::StartHandbrake(const FInputActionValue& Value)
 {
-	Super::BeginPlay();
-
-	bool bWantInCar = false;
-
-	// First disable both speed/gear displays 
-	bInCarCameraActive = false;
-	InCarSpeed->SetVisibility(bInCarCameraActive);
-	InCarGear->SetVisibility(bInCarCameraActive);
-
-	// Enable in car view if HMD is attached
-#if HMD_MODULE_INCLUDED
-	bWantInCar = UHeadMountedDisplayFunctionLibrary::IsHeadMountedDisplayEnabled();
-#endif // HMD_MODULE_INCLUDED
-
-	EnableIncarView(bWantInCar);
+	// add the input
+	ChaosVehicleMovement->SetHandbrakeInput(true);
 
-	// Start an engine sound playing
-	//EngineSoundComponent->Play();
+	// call the Blueprint hook for the break lights
+	BrakeLights(true);
 }
 
-void ATP_VehicleAdvPawn::OnResetVR()
+void ATP_VehicleAdvPawn::StopHandbrake(const FInputActionValue& Value)
 {
-#if HMD_MODULE_INCLUDED
-	if (GEngine->XRSystem.IsValid())
-	{
-		GEngine->XRSystem->ResetOrientationAndPosition();
-		InternalCamera->SetRelativeLocation(InternalCameraOrigin);
-		GetController()->SetControlRotation(FRotator());
-	}
-#endif // HMD_MODULE_INCLUDED
+	// add the input
+	ChaosVehicleMovement->SetHandbrakeInput(false);
+
+	// call the Blueprint hook for the break lights
+	BrakeLights(false);
 }
 
-void ATP_VehicleAdvPawn::UpdateHUDStrings()
+void ATP_VehicleAdvPawn::LookAround(const FInputActionValue& Value)
 {
-	float KPH = FMath::Abs(GetVehicleMovement()->GetForwardSpeed()) * 0.036f;
-	int32 KPH_int = FMath::FloorToInt(KPH);
-	int32 Gear = GetVehicleMovement()->GetCurrentGear();
-
-	// Using FText because this is display text that should be localizable
-	SpeedDisplayString = FText::Format(LOCTEXT("SpeedFormat", "{0} km/h"), FText::AsNumber(KPH_int));
-
-
-	if (bInReverseGear == true)
-	{
-		GearDisplayString = FText(LOCTEXT("ReverseGear", "R"));
-	}
-	else
-	{
-		GearDisplayString = (Gear == 0) ? LOCTEXT("N", "N") : FText::AsNumber(Gear);
-	}
+	// get the flat angle value for the input 
+	float LookValue = Value.Get<float>();
 
+	// add the input
+	BackSpringArm->AddLocalRotation(FRotator(0.0f, LookValue, 0.0f));
 }
 
-void ATP_VehicleAdvPawn::SetupInCarHUD()
+void ATP_VehicleAdvPawn::ToggleCamera(const FInputActionValue& Value)
 {
-	APlayerController* PlayerController = Cast<APlayerController>(GetController());
-	if ((PlayerController != nullptr) && (InCarSpeed != nullptr) && (InCarGear != nullptr))
-	{
-		// Setup the text render component strings
-		InCarSpeed->SetText(SpeedDisplayString);
-		InCarGear->SetText(GearDisplayString);
-		
-		if (bInReverseGear == false)
-		{
-			InCarGear->SetTextRenderColor(GearDisplayColor);
-		}
-		else
-		{
-			InCarGear->SetTextRenderColor(GearDisplayReverseColor);
-		}
-	}
+	// toggle the active camera flag
+	bFrontCameraActive = !bFrontCameraActive;
+
+	FrontCamera->SetActive(bFrontCameraActive);
+	BackCamera->SetActive(!bFrontCameraActive);
 }
 
-void ATP_VehicleAdvPawn::UpdatePhysicsMaterial()
+void ATP_VehicleAdvPawn::ResetVehicle(const FInputActionValue& Value)
 {
-	if (GetActorUpVector().Z < 0)
-	{
-		if (bIsLowFriction == true)
-		{
-			GetMesh()->SetPhysMaterialOverride(NonSlipperyMaterial);
-			bIsLowFriction = false;
-		}
-		else
-		{
-			GetMesh()->SetPhysMaterialOverride(SlipperyMaterial);
-			bIsLowFriction = true;
-		}
-	}
-}
+	// reset to a location slightly above our current one
+	FVector ResetLocation = GetActorLocation() + FVector(0.0f, 0.0f, 50.0f);
 
-#undef LOCTEXT_NAMESPACE
+	// reset to our yaw. Ignore pitch and roll
+	FRotator ResetRotation = GetActorRotation();
+	ResetRotation.Pitch = 0.0f;
+	ResetRotation.Roll = 0.0f;
+	
+	// teleport the actor to the reset spot and reset physics
+	SetActorTransform(FTransform(ResetRotation, ResetLocation, FVector::OneVector), false, nullptr, ETeleportType::TeleportPhysics);
+
+	GetMesh()->SetPhysicsAngularVelocityInDegrees(FVector::ZeroVector);
+	GetMesh()->SetPhysicsLinearVelocity(FVector::ZeroVector);
+
+	UE_LOG(LogTemplateVehicle, Error, TEXT("Reset Vehicle"));
+}
 
-PRAGMA_ENABLE_DEPRECATION_WARNINGS
+#undef LOCTEXT_NAMESPACE
\ No newline at end of file
diff --git a/Templates/TP_VehicleAdv/Source/TP_VehicleAdv/TP_VehicleAdvPawn.h b/Templates/TP_VehicleAdv/Source/TP_VehicleAdv/TP_VehicleAdvPawn.h
--- a/Templates/TP_VehicleAdv/Source/TP_VehicleAdv/TP_VehicleAdvPawn.h
+++ b/Templates/TP_VehicleAdv/Source/TP_VehicleAdv/TP_VehicleAdvPawn.h
@@ -6,146 +6,134 @@
 #include "WheeledVehiclePawn.h"
 #include "TP_VehicleAdvPawn.generated.h"
 
-class UPhysicalMaterial;
 class UCameraComponent;
 class USpringArmComponent;
-class UTextRenderComponent;
-class UInputComponent;
-class UAudioComponent;
-
-PRAGMA_DISABLE_DEPRECATION_WARNINGS
-
-UCLASS(config=Game)
+class UInputAction;
+class UChaosWheeledVehicleMovementComponent;
+struct FInputActionValue;
+
+DECLARE_LOG_CATEGORY_EXTERN(LogTemplateVehicle, Log, All);
+
+/**
+ *  Vehicle Pawn class
+ *  Handles common functionality for all vehicle types,
+ *  including input handling and camera management.
+ *  
+ *  Specific vehicle configurations are handled in subclasses.
+ */
+UCLASS(abstract)
 class ATP_VehicleAdvPawn : public AWheeledVehiclePawn
 {
 	GENERATED_BODY()
 
-	/** Spring arm that will offset the camera */
-	UPROPERTY(Category = Camera, VisibleDefaultsOnly, BlueprintReadOnly, meta = (AllowPrivateAccess = "true"))
-	USpringArmComponent* SpringArm;
+	/** Spring Arm for the front camera */
+	UPROPERTY(VisibleAnywhere, BlueprintReadOnly, Category = Camera, meta = (AllowPrivateAccess = "true"))
+	USpringArmComponent* FrontSpringArm;
 
-	/** Camera component that will be our viewpoint */
-	UPROPERTY(Category = Camera, VisibleDefaultsOnly, BlueprintReadOnly, meta = (AllowPrivateAccess = "true"))
-	UCameraComponent* Camera;
+	/** Front Camera component */
+	UPROPERTY(VisibleAnywhere, BlueprintReadOnly, Category = Camera, meta = (AllowPrivateAccess = "true"))
+	UCameraComponent* FrontCamera;
 
-	/** SCene component for the In-Car view origin */
-	UPROPERTY(Category = Camera, VisibleDefaultsOnly, BlueprintReadOnly, meta = (AllowPrivateAccess = "true"))
-	class USceneComponent* InternalCameraBase;
+	/** Spring Arm for the back camera */
+	UPROPERTY(VisibleAnywhere, BlueprintReadOnly, Category = Camera, meta = (AllowPrivateAccess = "true"))
+	USpringArmComponent* BackSpringArm;
 
-	/** Camera component for the In-Car view */
-	UPROPERTY(Category = Camera, VisibleDefaultsOnly, BlueprintReadOnly, meta = (AllowPrivateAccess = "true"))
-	UCameraComponent* InternalCamera;
+	/** Back Camera component */
+	UPROPERTY(VisibleAnywhere, BlueprintReadOnly, Category = Camera, meta = (AllowPrivateAccess = "true"))
+	UCameraComponent* BackCamera;
 
-	/** Text component for the In-Car speed */
-	UPROPERTY(Category = Display, VisibleDefaultsOnly, BlueprintReadOnly, meta = (AllowPrivateAccess = "true"))
-	UTextRenderComponent* InCarSpeed;
+	/** Cast pointer to the Chaos Vehicle movement component */
+	TObjectPtr<UChaosWheeledVehicleMovementComponent> ChaosVehicleMovement;
 
-	/** Text component for the In-Car gear */
-	UPROPERTY(Category = Display, VisibleDefaultsOnly, BlueprintReadOnly, meta = (AllowPrivateAccess = "true"))
-	UTextRenderComponent* InCarGear;
+protected:
 
-	/** Audio component for the engine sound */
-	UPROPERTY(Category = Display, VisibleDefaultsOnly, BlueprintReadOnly, meta = (AllowPrivateAccess = "true"))
-	UAudioComponent* EngineSoundComponent;
+	/** Steering Action */
+	UPROPERTY(EditAnywhere, BlueprintReadOnly, Category = Input)
+	UInputAction* SteeringAction;
 
-public:
-	ATP_VehicleAdvPawn();
+	/** Throttle Action */
+	UPROPERTY(EditAnywhere, BlueprintReadOnly, Category = Input)
+	UInputAction* ThrottleAction;
 
-	/** The current speed as a string eg 10 km/h */
-	UPROPERTY(Category = Display, VisibleDefaultsOnly, BlueprintReadOnly)
-	FText SpeedDisplayString;
+	/** Brake Action */
+	UPROPERTY(EditAnywhere, BlueprintReadOnly, Category = Input)
+	UInputAction* BrakeAction;
 
-	/** The current gear as a string (R,N, 1,2 etc) */
-	UPROPERTY(Category = Display, VisibleDefaultsOnly, BlueprintReadOnly)
-	FText GearDisplayString;
+	/** Handbrake Action */
+	UPROPERTY(EditAnywhere, BlueprintReadOnly, Category = Input)
+	UInputAction* HandbrakeAction;
 
-	UPROPERTY(Category = Display, VisibleDefaultsOnly, BlueprintReadOnly)
-	/** The color of the incar gear text in forward gears */
-	FColor	GearDisplayColor;
+	/** Look Around Action */
+	UPROPERTY(EditAnywhere, BlueprintReadOnly, Category = Input)
+	UInputAction* LookAroundAction;
 
-	/** The color of the incar gear text when in reverse */
-	UPROPERTY(Category = Display, VisibleDefaultsOnly, BlueprintReadOnly)
-	FColor	GearDisplayReverseColor;
+	/** Toggle Camera Action */
+	UPROPERTY(EditAnywhere, BlueprintReadOnly, Category = Input)
+	UInputAction* ToggleCameraAction;
 
-	/** Are we using incar camera */
-	UPROPERTY(Category = Camera, VisibleDefaultsOnly, BlueprintReadOnly)
-	bool bInCarCameraActive;
+	/** Reset Vehicle Action */
+	UPROPERTY(EditAnywhere, BlueprintReadOnly, Category = Input)
+	UInputAction* ResetVehicleAction;
 
-	/** Are we in reverse gear */
-	UPROPERTY(Category = Camera, VisibleDefaultsOnly, BlueprintReadOnly)
-	bool bInReverseGear;
+	/** Keeps track of which camera is active */
+	bool bFrontCameraActive = false;
 
-	/** Initial offset of incar camera */
-	FVector InternalCameraOrigin;
+public:
+	ATP_VehicleAdvPawn();
 
 	// Begin Pawn interface
+
 	virtual void SetupPlayerInputComponent(UInputComponent* InputComponent) override;
+
 	// End Pawn interface
 
 	// Begin Actor interface
+
 	virtual void Tick(float Delta) override;
-protected:
-	virtual void BeginPlay() override;
 
-public:
 	// End Actor interface
 
-	/** Handle pressing forwards */
-	void MoveForward(float Val);
-
-	/** Setup the strings used on the hud */
-	void SetupInCarHUD();
-
-	/** Update the physics material used by the vehicle mesh */
-	void UpdatePhysicsMaterial();
-
-	/** Handle pressing right */
-	void MoveRight(float Val);
-	/** Handle handbrake pressed */
-	void OnHandbrakePressed();
-	/** Handle handbrake released */
-	void OnHandbrakeReleased();
-	/** Switch between cameras */
-	void OnToggleCamera();
-	/** Handle reset VR device */
-	void OnResetVR();
-
-	static const FName LookUpBinding;
-	static const FName LookRightBinding;
-	static const FName EngineAudioRPM;
-
-private:
-	/** 
-	 * Activate In-Car camera. Enable camera and sets visibility of incar hud display
-	 *
-	 * @param	bState true will enable in car view and set visibility of various
-	 */
-	void EnableIncarView( const bool bState );
-
-	/** Update the gear and speed strings */
-	void UpdateHUDStrings();
-
-	/* Are we on a 'slippery' surface */
-	bool bIsLowFriction;
-	/** Slippery Material instance */
-	UPhysicalMaterial* SlipperyMaterial;
-	/** Non Slippery Material instance */
-	UPhysicalMaterial* NonSlipperyMaterial;
+protected:
+
+	/** Handles steering input */
+	void Steering(const FInputActionValue& Value);
+
+	/** Handles throttle input */
+	void Throttle(const FInputActionValue& Value);
+
+	/** Handles brake input */
+	void Brake(const FInputActionValue& Value);
+
+	/** Handles brake start/stop inputs */
+	void StartBrake(const FInputActionValue& Value);
+	void StopBrake(const FInputActionValue& Value);
 
+	/** Handles handbrake start/stop inputs */
+	void StartHandbrake(const FInputActionValue& Value);
+	void StopHandbrake(const FInputActionValue& Value);
+
+	/** Handles look around input */
+	void LookAround(const FInputActionValue& Value);
+
+	/** Handles toggle camera input */
+	void ToggleCamera(const FInputActionValue& Value);
+
+	/** Handles reset vehicle input */
+	void ResetVehicle(const FInputActionValue& Value);
+
+	/** Called when the brake lights are turned on or off */
+	UFUNCTION(BlueprintImplementableEvent, Category="Vehicle")
+	void BrakeLights(bool bBraking);
 
 public:
-	/** Returns SpringArm subobject **/
-	FORCEINLINE USpringArmComponent* GetSpringArm() const { return SpringArm; }
-	/** Returns Camera subobject **/
-	FORCEINLINE UCameraComponent* GetCamera() const { return Camera; }
-	/** Returns InternalCamera subobject **/
-	FORCEINLINE UCameraComponent* GetInternalCamera() const { return InternalCamera; }
-	/** Returns InCarSpeed subobject **/
-	FORCEINLINE UTextRenderComponent* GetInCarSpeed() const { return InCarSpeed; }
-	/** Returns InCarGear subobject **/
-	FORCEINLINE UTextRenderComponent* GetInCarGear() const { return InCarGear; }
-	/** Returns EngineSoundComponent subobject **/
-	FORCEINLINE UAudioComponent* GetEngineSoundComponent() const { return EngineSoundComponent; }
+	/** Returns the front spring arm subobject */
+	FORCEINLINE USpringArmComponent* GetFrontSpringArm() const { return FrontSpringArm; }
+	/** Returns the front camera subobject */
+	FORCEINLINE UCameraComponent* GetFollowCamera() const { return FrontCamera; }
+	/** Returns the back spring arm subobject */
+	FORCEINLINE USpringArmComponent* GetBackSpringArm() const { return BackSpringArm; }
+	/** Returns the back camera subobject */
+	FORCEINLINE UCameraComponent* GetBackCamera() const { return BackCamera; }
+	/** Returns the cast Chaos Vehicle Movement subobject */
+	FORCEINLINE const TObjectPtr<UChaosWheeledVehicleMovementComponent>& GetChaosVehicleMovement() const { return ChaosVehicleMovement; }
 };
-
-PRAGMA_ENABLE_DEPRECATION_WARNINGS
diff --git a/Templates/TP_VehicleAdv/Source/TP_VehicleAdv/TP_VehicleAdvPlayerController.cpp b/Templates/TP_VehicleAdv/Source/TP_VehicleAdv/TP_VehicleAdvPlayerController.cpp
new file mode 100644
--- /dev/null
+++ b/Templates/TP_VehicleAdv/Source/TP_VehicleAdv/TP_VehicleAdvPlayerController.cpp
@@ -0,0 +1,46 @@
+// Copyright Epic Games, Inc. All Rights Reserved.
+
+
+#include "TP_VehicleAdvPlayerController.h"
+#include "TP_VehicleAdvPawn.h"
+#include "TP_VehicleAdvUI.h"
+#include "EnhancedInputSubsystems.h"
+#include "ChaosWheeledVehicleMovementComponent.h"
+
+void ATP_VehicleAdvPlayerController::BeginPlay()
+{
+	Super::BeginPlay();
+
+	// get the enhanced input subsystem
+	if (UEnhancedInputLocalPlayerSubsystem* Subsystem = ULocalPlayer::GetSubsystem<UEnhancedInputLocalPlayerSubsystem>(GetLocalPlayer()))
+	{
+		// add the mapping context so we get controls
+		Subsystem->AddMappingContext(InputMappingContext, 0);
+	}
+
+	// spawn the UI widget and add it to the viewport
+	VehicleUI = CreateWidget<UTP_VehicleAdvUI>(this, VehicleUIClass);
+
+	check(VehicleUI);
+
+	VehicleUI->AddToViewport();
+}
+
+void ATP_VehicleAdvPlayerController::Tick(float Delta)
+{
+	Super::Tick(Delta);
+
+	if (IsValid(VehiclePawn) && IsValid(VehicleUI))
+	{
+		VehicleUI->UpdateSpeed(VehiclePawn->GetChaosVehicleMovement()->GetForwardSpeed());
+		VehicleUI->UpdateGear(VehiclePawn->GetChaosVehicleMovement()->GetCurrentGear());
+	}
+}
+
+void ATP_VehicleAdvPlayerController::OnPossess(APawn* InPawn)
+{
+	Super::OnPossess(InPawn);
+
+	// get a pointer to the controlled pawn
+	VehiclePawn = CastChecked<ATP_VehicleAdvPawn>(InPawn);
+}
diff --git a/Templates/TP_VehicleAdv/Source/TP_VehicleAdv/TP_VehicleAdvPlayerController.h b/Templates/TP_VehicleAdv/Source/TP_VehicleAdv/TP_VehicleAdvPlayerController.h
new file mode 100644
--- /dev/null
+++ b/Templates/TP_VehicleAdv/Source/TP_VehicleAdv/TP_VehicleAdvPlayerController.h
@@ -0,0 +1,55 @@
+// Copyright Epic Games, Inc. All Rights Reserved.
+
+#pragma once
+
+#include "CoreMinimal.h"
+#include "GameFramework/PlayerController.h"
+#include "TP_VehicleAdvPlayerController.generated.h"
+
+class UInputMappingContext;
+class ATP_VehicleAdvPawn;
+class UTP_VehicleAdvUI;
+
+/**
+ *  Vehicle Player Controller class
+ *  Handles input mapping and user interface
+ */
+UCLASS(abstract)
+class TP_VEHICLEADV_API ATP_VehicleAdvPlayerController : public APlayerController
+{
+	GENERATED_BODY()
+
+protected:
+
+	/** Input Mapping Context to be used for player input */
+	UPROPERTY(EditAnywhere, BlueprintReadOnly, Category = Input)
+	UInputMappingContext* InputMappingContext;
+
+	/** Pointer to the controlled vehicle pawn */
+	TObjectPtr<ATP_VehicleAdvPawn> VehiclePawn;
+
+	/** Type of the UI to spawn */
+	UPROPERTY(EditAnywhere, BlueprintReadOnly, Category = UI)
+	TSubclassOf<UTP_VehicleAdvUI> VehicleUIClass;
+
+	/** Pointer to the UI widget */
+	TObjectPtr<UTP_VehicleAdvUI> VehicleUI;
+
+	// Begin Actor interface
+protected:
+
+	virtual void BeginPlay() override;
+
+public:
+
+	virtual void Tick(float Delta) override;
+
+	// End Actor interface
+
+	// Begin PlayerController interface
+protected:
+
+	virtual void OnPossess(APawn* InPawn) override;
+
+	// End PlayerController interface
+};
diff --git a/Templates/TP_VehicleAdv/Source/TP_VehicleAdv/TP_VehicleAdvSportsCar.cpp b/Templates/TP_VehicleAdv/Source/TP_VehicleAdv/TP_VehicleAdvSportsCar.cpp
new file mode 100644
--- /dev/null
+++ b/Templates/TP_VehicleAdv/Source/TP_VehicleAdv/TP_VehicleAdvSportsCar.cpp
@@ -0,0 +1,69 @@
+// Copyright Epic Games, Inc. All Rights Reserved.
+
+
+#include "TP_VehicleAdvSportsCar.h"
+#include "TP_VehicleAdvSportsWheelFront.h"
+#include "TP_VehicleAdvSportsWheelRear.h"
+#include "ChaosWheeledVehicleMovementComponent.h"
+
+ATP_VehicleAdvSportsCar::ATP_VehicleAdvSportsCar()
+{
+	// Note: for faster iteration times, the vehicle setup can be tweaked in the Blueprint instead
+
+	// Set up the chassis
+	GetChaosVehicleMovement()->ChassisHeight = 144.0f;
+	GetChaosVehicleMovement()->DragCoefficient = 0.31f;
+
+	// Set up the wheels
+	GetChaosVehicleMovement()->bLegacyWheelFrictionPosition = true;
+	GetChaosVehicleMovement()->WheelSetups.SetNum(4);
+
+	GetChaosVehicleMovement()->WheelSetups[0].WheelClass = UTP_VehicleAdvSportsWheelFront::StaticClass();
+	GetChaosVehicleMovement()->WheelSetups[0].BoneName = FName("Phys_Wheel_FL");
+	GetChaosVehicleMovement()->WheelSetups[0].AdditionalOffset = FVector(0.0f, 0.0f, 0.0f);
+
+	GetChaosVehicleMovement()->WheelSetups[1].WheelClass = UTP_VehicleAdvSportsWheelFront::StaticClass();
+	GetChaosVehicleMovement()->WheelSetups[1].BoneName = FName("Phys_Wheel_FR");
+	GetChaosVehicleMovement()->WheelSetups[1].AdditionalOffset = FVector(0.0f, 0.0f, 0.0f);
+
+	GetChaosVehicleMovement()->WheelSetups[2].WheelClass = UTP_VehicleAdvSportsWheelRear::StaticClass();
+	GetChaosVehicleMovement()->WheelSetups[2].BoneName = FName("Phys_Wheel_BL");
+	GetChaosVehicleMovement()->WheelSetups[2].AdditionalOffset = FVector(0.0f, 0.0f, 0.0f);
+
+	GetChaosVehicleMovement()->WheelSetups[3].WheelClass = UTP_VehicleAdvSportsWheelRear::StaticClass();
+	GetChaosVehicleMovement()->WheelSetups[3].BoneName = FName("Phys_Wheel_BR");
+	GetChaosVehicleMovement()->WheelSetups[3].AdditionalOffset = FVector(0.0f, 0.0f, 0.0f);
+
+	// Set up the engine
+	// NOTE: Check the Blueprint asset for the Torque Curve
+	GetChaosVehicleMovement()->EngineSetup.MaxTorque = 750.0f;
+	GetChaosVehicleMovement()->EngineSetup.MaxRPM = 7000.0f;
+	GetChaosVehicleMovement()->EngineSetup.EngineIdleRPM = 900.0f;
+	GetChaosVehicleMovement()->EngineSetup.EngineBrakeEffect = 0.2f;
+	GetChaosVehicleMovement()->EngineSetup.EngineRevUpMOI = 5.0f;
+	GetChaosVehicleMovement()->EngineSetup.EngineRevDownRate = 600.0f;
+
+	// Set up the transmission
+	GetChaosVehicleMovement()->TransmissionSetup.bUseAutomaticGears = true;
+	GetChaosVehicleMovement()->TransmissionSetup.bUseAutoReverse = true;
+	GetChaosVehicleMovement()->TransmissionSetup.FinalRatio = 2.81f;
+	GetChaosVehicleMovement()->TransmissionSetup.ChangeUpRPM = 6000.0f;
+	GetChaosVehicleMovement()->TransmissionSetup.ChangeDownRPM = 2000.0f;
+	GetChaosVehicleMovement()->TransmissionSetup.GearChangeTime = 0.2f;
+	GetChaosVehicleMovement()->TransmissionSetup.TransmissionEfficiency = 0.9f;
+
+	GetChaosVehicleMovement()->TransmissionSetup.ForwardGearRatios.SetNum(5);
+	GetChaosVehicleMovement()->TransmissionSetup.ForwardGearRatios[0] = 4.25f;
+	GetChaosVehicleMovement()->TransmissionSetup.ForwardGearRatios[1] = 2.52f;
+	GetChaosVehicleMovement()->TransmissionSetup.ForwardGearRatios[2] = 1.66f;
+	GetChaosVehicleMovement()->TransmissionSetup.ForwardGearRatios[3] = 1.22f;
+	GetChaosVehicleMovement()->TransmissionSetup.ForwardGearRatios[4] = 1.0f;
+
+	GetChaosVehicleMovement()->TransmissionSetup.ReverseGearRatios.SetNum(1);
+	GetChaosVehicleMovement()->TransmissionSetup.ReverseGearRatios[0] = 4.04f;
+
+	// Set up the steering
+	// NOTE: Check the Blueprint asset for the Steering Curve
+	GetChaosVehicleMovement()->SteeringSetup.SteeringType = ESteeringType::Ackermann;
+	GetChaosVehicleMovement()->SteeringSetup.AngleRatio = 0.7f;
+}
\ No newline at end of file
diff --git a/Templates/TP_VehicleAdv/Source/TP_VehicleAdv/TP_VehicleAdvSportsCar.h b/Templates/TP_VehicleAdv/Source/TP_VehicleAdv/TP_VehicleAdvSportsCar.h
new file mode 100644
--- /dev/null
+++ b/Templates/TP_VehicleAdv/Source/TP_VehicleAdv/TP_VehicleAdvSportsCar.h
@@ -0,0 +1,20 @@
+// Copyright Epic Games, Inc. All Rights Reserved.
+
+#pragma once
+
+#include "CoreMinimal.h"
+#include "TP_VehicleAdvPawn.h"
+#include "TP_VehicleAdvSportsCar.generated.h"
+
+/**
+ *  Sports car wheeled vehicle implementation
+ */
+UCLASS(abstract)
+class TP_VEHICLEADV_API ATP_VehicleAdvSportsCar : public ATP_VehicleAdvPawn
+{
+	GENERATED_BODY()
+	
+public:
+
+	ATP_VehicleAdvSportsCar();
+};
diff --git a/Templates/TP_VehicleAdv/Source/TP_VehicleAdv/TP_VehicleAdvSportsWheelFront.cpp b/Templates/TP_VehicleAdv/Source/TP_VehicleAdv/TP_VehicleAdvSportsWheelFront.cpp
new file mode 100644
--- /dev/null
+++ b/Templates/TP_VehicleAdv/Source/TP_VehicleAdv/TP_VehicleAdvSportsWheelFront.cpp
@@ -0,0 +1,14 @@
+// Copyright Epic Games, Inc. All Rights Reserved.
+
+
+#include "TP_VehicleAdvSportsWheelFront.h"
+
+UTP_VehicleAdvSportsWheelFront::UTP_VehicleAdvSportsWheelFront()
+{
+	WheelRadius = 39.0f;
+	WheelWidth = 35.0f;
+	FrictionForceMultiplier = 3.0f;
+
+	MaxBrakeTorque = 4500.0f;
+	MaxHandBrakeTorque = 6000.0f;
+}
\ No newline at end of file
diff --git a/Templates/TP_VehicleAdv/Source/TP_VehicleAdv/TP_VehicleAdvSportsWheelFront.h b/Templates/TP_VehicleAdv/Source/TP_VehicleAdv/TP_VehicleAdvSportsWheelFront.h
new file mode 100644
--- /dev/null
+++ b/Templates/TP_VehicleAdv/Source/TP_VehicleAdv/TP_VehicleAdvSportsWheelFront.h
@@ -0,0 +1,20 @@
+// Copyright Epic Games, Inc. All Rights Reserved.
+
+#pragma once
+
+#include "CoreMinimal.h"
+#include "TP_VehicleAdvWheelFront.h"
+#include "TP_VehicleAdvSportsWheelFront.generated.h"
+
+/**
+ *  Front wheel definition for Sports Car.
+ */
+UCLASS()
+class TP_VEHICLEADV_API UTP_VehicleAdvSportsWheelFront : public UTP_VehicleAdvWheelFront
+{
+	GENERATED_BODY()
+	
+public:
+
+	UTP_VehicleAdvSportsWheelFront();
+};
diff --git a/Templates/TP_VehicleAdv/Source/TP_VehicleAdv/TP_VehicleAdvSportsWheelRear.cpp b/Templates/TP_VehicleAdv/Source/TP_VehicleAdv/TP_VehicleAdvSportsWheelRear.cpp
new file mode 100644
--- /dev/null
+++ b/Templates/TP_VehicleAdv/Source/TP_VehicleAdv/TP_VehicleAdvSportsWheelRear.cpp
@@ -0,0 +1,15 @@
+// Copyright Epic Games, Inc. All Rights Reserved.
+
+
+#include "TP_VehicleAdvSportsWheelRear.h"
+
+UTP_VehicleAdvSportsWheelRear::UTP_VehicleAdvSportsWheelRear()
+{
+	WheelRadius = 40.f;
+	WheelWidth = 40.0f;
+	FrictionForceMultiplier = 4.0f;
+	SlipThreshold = 100.0f;
+	SkidThreshold = 100.0f;
+	MaxSteerAngle = 0.0f;
+	MaxHandBrakeTorque = 6000.0f;
+}
\ No newline at end of file
diff --git a/Templates/TP_VehicleAdv/Source/TP_VehicleAdv/TP_VehicleAdvSportsWheelRear.h b/Templates/TP_VehicleAdv/Source/TP_VehicleAdv/TP_VehicleAdvSportsWheelRear.h
new file mode 100644
--- /dev/null
+++ b/Templates/TP_VehicleAdv/Source/TP_VehicleAdv/TP_VehicleAdvSportsWheelRear.h
@@ -0,0 +1,20 @@
+// Copyright Epic Games, Inc. All Rights Reserved.
+
+#pragma once
+
+#include "CoreMinimal.h"
+#include "TP_VehicleAdvWheelRear.h"
+#include "TP_VehicleAdvSportsWheelRear.generated.h"
+
+/**
+ *  Rear wheel definition for Sports Car.
+ */
+UCLASS()
+class TP_VEHICLEADV_API UTP_VehicleAdvSportsWheelRear : public UTP_VehicleAdvWheelRear
+{
+	GENERATED_BODY()
+	
+public:
+
+	UTP_VehicleAdvSportsWheelRear();
+};
diff --git a/Templates/TP_VehicleAdv/Source/TP_VehicleAdv/TP_VehicleAdvUI.cpp b/Templates/TP_VehicleAdv/Source/TP_VehicleAdv/TP_VehicleAdvUI.cpp
new file mode 100644
--- /dev/null
+++ b/Templates/TP_VehicleAdv/Source/TP_VehicleAdv/TP_VehicleAdvUI.cpp
@@ -0,0 +1,19 @@
+// Copyright Epic Games, Inc. All Rights Reserved.
+
+
+#include "TP_VehicleAdvUI.h"
+
+void UTP_VehicleAdvUI::UpdateSpeed(float NewSpeed)
+{
+	// format the speed to KPH or MPH
+	float FormattedSpeed = FMath::Abs(NewSpeed) * (bIsMPH ? 0.022f : 0.036f);
+
+	// call the Blueprint handler
+	OnSpeedUpdate(FormattedSpeed);
+}
+
+void UTP_VehicleAdvUI::UpdateGear(int32 NewGear)
+{
+	// call the Blueprint handler
+	OnGearUpdate(NewGear);
+}
\ No newline at end of file
diff --git a/Templates/TP_VehicleAdv/Source/TP_VehicleAdv/TP_VehicleAdvUI.h b/Templates/TP_VehicleAdv/Source/TP_VehicleAdv/TP_VehicleAdvUI.h
new file mode 100644
--- /dev/null
+++ b/Templates/TP_VehicleAdv/Source/TP_VehicleAdv/TP_VehicleAdvUI.h
@@ -0,0 +1,42 @@
+// Copyright Epic Games, Inc. All Rights Reserved.
+
+#pragma once
+
+#include "CoreMinimal.h"
+#include "Blueprint/UserWidget.h"
+#include "TP_VehicleAdvUI.generated.h"
+
+/**
+ *  Simple Vehicle HUD class
+ *  Displays the current speed and gear.
+ *  Widget setup is handled in a Blueprint subclass.
+ */
+UCLASS(abstract)
+class TP_VEHICLEADV_API UTP_VehicleAdvUI : public UUserWidget
+{
+	GENERATED_BODY()
+	
+protected:
+
+	/** Controls the display of speed in Km/h or MPH */
+	UPROPERTY(EditAnywhere, BlueprintReadOnly, Category = Vehicle)
+	bool bIsMPH = false;
+
+public:
+
+	/** Called to update the speed display */
+	void UpdateSpeed(float NewSpeed);
+
+	/** Called to update the gear display */
+	void UpdateGear(int32 NewGear);
+
+protected:
+
+	/** Implemented in Blueprint to display the new speed */
+	UFUNCTION(BlueprintImplementableEvent, Category = Vehicle)
+	void OnSpeedUpdate(float NewSpeed);
+
+	/** Implemented in Blueprint to display the new gear */
+	UFUNCTION(BlueprintImplementableEvent, Category = Vehicle)
+	void OnGearUpdate(int32 NewGear);
+};
diff --git a/Templates/TP_VehicleAdv/Source/TP_VehicleAdv/TP_VehicleAdvWheelFront.cpp b/Templates/TP_VehicleAdv/Source/TP_VehicleAdv/TP_VehicleAdvWheelFront.cpp
--- a/Templates/TP_VehicleAdv/Source/TP_VehicleAdv/TP_VehicleAdvWheelFront.cpp
+++ b/Templates/TP_VehicleAdv/Source/TP_VehicleAdv/TP_VehicleAdvWheelFront.cpp
@@ -3,27 +3,9 @@
 #include "TP_VehicleAdvWheelFront.h"
 #include "UObject/ConstructorHelpers.h"
 
-PRAGMA_DISABLE_DEPRECATION_WARNINGS
-
 UTP_VehicleAdvWheelFront::UTP_VehicleAdvWheelFront()
 {
-	WheelRadius = 39.f;
-	WheelWidth = 35.0f;
-	FrictionForceMultiplier = 3.0f;
-	CorneringStiffness = 1200.0f;
-	bAffectedByEngine = false;
-	bAffectedByHandbrake = false;
-	bAffectedBySteering = true;
 	AxleType = EAxleType::Front;
-	SpringRate = 250.0f;
-	SpringPreload = 50.f;
-	SuspensionDampingRatio = 0.5f;
-	WheelLoadRatio = 0.5f;
-	RollbarScaling = 0.15f;
-	SuspensionMaxRaise = 10.0f;
-	SuspensionMaxDrop = 10.0f;
-	WheelLoadRatio = 0.5f;
+	bAffectedBySteering = true;
 	MaxSteerAngle = 40.f;
-}
-
-PRAGMA_ENABLE_DEPRECATION_WARNINGS
+}
\ No newline at end of file
diff --git a/Templates/TP_VehicleAdv/Source/TP_VehicleAdv/TP_VehicleAdvWheelFront.h b/Templates/TP_VehicleAdv/Source/TP_VehicleAdv/TP_VehicleAdvWheelFront.h
--- a/Templates/TP_VehicleAdv/Source/TP_VehicleAdv/TP_VehicleAdvWheelFront.h
+++ b/Templates/TP_VehicleAdv/Source/TP_VehicleAdv/TP_VehicleAdvWheelFront.h
@@ -6,8 +6,9 @@
 #include "ChaosVehicleWheel.h"
 #include "TP_VehicleAdvWheelFront.generated.h"
 
-PRAGMA_DISABLE_DEPRECATION_WARNINGS
-
+/**
+ *  Base front wheel definition.
+ */
 UCLASS()
 class UTP_VehicleAdvWheelFront : public UChaosVehicleWheel
 {
@@ -15,6 +16,4 @@ class UTP_VehicleAdvWheelFront : public UChaosVehicleWheel
 
 public:
 	UTP_VehicleAdvWheelFront();
-};
-
-PRAGMA_ENABLE_DEPRECATION_WARNINGS
+};
\ No newline at end of file
diff --git a/Templates/TP_VehicleAdv/Source/TP_VehicleAdv/TP_VehicleAdvWheelRear.cpp b/Templates/TP_VehicleAdv/Source/TP_VehicleAdv/TP_VehicleAdvWheelRear.cpp
--- a/Templates/TP_VehicleAdv/Source/TP_VehicleAdv/TP_VehicleAdvWheelRear.cpp
+++ b/Templates/TP_VehicleAdv/Source/TP_VehicleAdv/TP_VehicleAdvWheelRear.cpp
@@ -3,26 +3,9 @@
 #include "TP_VehicleAdvWheelRear.h"
 #include "UObject/ConstructorHelpers.h"
 
-PRAGMA_DISABLE_DEPRECATION_WARNINGS
-
 UTP_VehicleAdvWheelRear::UTP_VehicleAdvWheelRear()
 {
-	WheelRadius = 40.f;
-	WheelWidth = 40.0f;
-	FrictionForceMultiplier = 4.0f;
-	CorneringStiffness = 1200.0f;
-	bAffectedByEngine = true;
-	bAffectedByHandbrake = true;
-	bAffectedBySteering = false;
 	AxleType = EAxleType::Rear;
-	SpringRate = 250.0f;
-	SpringPreload = 50.f;
-	SuspensionDampingRatio = 0.5f;
-	WheelLoadRatio = 0.5f;
-	RollbarScaling = 0.5f;
-	SuspensionMaxRaise = 10.0f;
-	SuspensionMaxDrop = 10.0f;
-	WheelLoadRatio = 0.5f;
-}
-
-PRAGMA_ENABLE_DEPRECATION_WARNINGS
+	bAffectedByHandbrake = true;
+	bAffectedByEngine = true;
+}
\ No newline at end of file
diff --git a/Templates/TP_VehicleAdv/Source/TP_VehicleAdv/TP_VehicleAdvWheelRear.h b/Templates/TP_VehicleAdv/Source/TP_VehicleAdv/TP_VehicleAdvWheelRear.h
--- a/Templates/TP_VehicleAdv/Source/TP_VehicleAdv/TP_VehicleAdvWheelRear.h
+++ b/Templates/TP_VehicleAdv/Source/TP_VehicleAdv/TP_VehicleAdvWheelRear.h
@@ -6,8 +6,9 @@
 #include "ChaosVehicleWheel.h"
 #include "TP_VehicleAdvWheelRear.generated.h"
 
-PRAGMA_DISABLE_DEPRECATION_WARNINGS
-
+/**
+ *  Base rear wheel definition.
+ */
 UCLASS()
 class UTP_VehicleAdvWheelRear : public UChaosVehicleWheel
 {
@@ -16,7 +17,3 @@ class UTP_VehicleAdvWheelRear : public UChaosVehicleWheel
 public:
 	UTP_VehicleAdvWheelRear();
 };
-
-PRAGMA_ENABLE_DEPRECATION_WARNINGS
-
-
diff --git a/Templates/TP_VehicleAdv/TP_VehicleAdv.uproject b/Templates/TP_VehicleAdv/TP_VehicleAdv.uproject
--- a/Templates/TP_VehicleAdv/TP_VehicleAdv.uproject
+++ b/Templates/TP_VehicleAdv/TP_VehicleAdv.uproject
@@ -7,28 +7,28 @@
 		{
 			"Name": "TP_VehicleAdv",
 			"Type": "Runtime",
-			"LoadingPhase": "Default"
+			"LoadingPhase": "Default",
+			"AdditionalDependencies": [
+				"Engine",
+				"UMG"
+			]
 		}
 	],
 	"Plugins": [
-		{
-			"Name": "PhysXVehicles",
-			"Enabled": false
-		},
 		{
 			"Name": "ChaosVehiclesPlugin",
 			"Enabled": true
 		},
-		{
-			"Name": "RawInput",
-			"Enabled": true
-		},
 		{
 			"Name": "ModelingToolsEditorMode",
 			"Enabled": true,
 			"TargetAllowList": [
 				"Editor"
 			]
-		}		
+		},
+		{
+			"Name": "RawInput",
+			"Enabled": true
+		}
 	]
 }
\ No newline at end of file
