Templates - Separated steering wheel input mappings into a separate context to avoid an issue with RawInput

diff --git a/Templates/TP_VehicleAdv/Source/TP_VehicleAdv/TP_VehicleAdvPlayerController.cpp b/Templates/TP_VehicleAdv/Source/TP_VehicleAdv/TP_VehicleAdvPlayerController.cpp
--- a/Templates/TP_VehicleAdv/Source/TP_VehicleAdv/TP_VehicleAdvPlayerController.cpp
+++ b/Templates/TP_VehicleAdv/Source/TP_VehicleAdv/TP_VehicleAdvPlayerController.cpp
@@ -28,6 +28,12 @@ void ATP_VehicleAdvPlayerController::SetupInputComponent()
 	{
 		// add the mapping context so we get controls
 		Subsystem->AddMappingContext(InputMappingContext, 0);
+
+		// optionally add the steering wheel context
+		if (bUseSteeringWheelControls && SteeringWheelInputMappingContext)
+		{
+			Subsystem->AddMappingContext(SteeringWheelInputMappingContext, 1);
+		}
 	}
 }
 
diff --git a/Templates/TP_VehicleAdv/Source/TP_VehicleAdv/TP_VehicleAdvPlayerController.h b/Templates/TP_VehicleAdv/Source/TP_VehicleAdv/TP_VehicleAdvPlayerController.h
--- a/Templates/TP_VehicleAdv/Source/TP_VehicleAdv/TP_VehicleAdvPlayerController.h
+++ b/Templates/TP_VehicleAdv/Source/TP_VehicleAdv/TP_VehicleAdvPlayerController.h
@@ -25,6 +25,16 @@ protected:
 	UPROPERTY(EditAnywhere, BlueprintReadOnly, Category = Input)
 	UInputMappingContext* InputMappingContext;
 
+	/** If true, the optional steering wheel input mapping context will be registered */
+	UPROPERTY(EditAnywhere, BlueprintReadOnly, Category = Input)
+	bool bUseSteeringWheelControls = false;
+
+	/** Optional Input Mapping Context to be used for steering wheel input.
+	 *  This is added alongside the default Input Mapping Context and does not block other forms of input.
+	 */
+	UPROPERTY(EditAnywhere, BlueprintReadOnly, Category = Input, meta=(EditCondition="bUseSteeringWheelControls"))
+	UInputMappingContext* SteeringWheelInputMappingContext;
+
 	/** Pointer to the controlled vehicle pawn */
 	TObjectPtr<ATP_VehicleAdvPawn> VehiclePawn;
 
@@ -35,6 +45,8 @@ protected:
 	/** Pointer to the UI widget */
 	TObjectPtr<UTP_VehicleAdvUI> VehicleUI;
 
+	
+
 	// Begin Actor interface
 protected:
 
