Templates - TP_VehicleAdv
Updated template code and content.
Implemented Time Trial variant
Implemented Offroad variant
#rnx

diff --git a/Templates/TP_VehicleAdv/Config/DefaultEngine.ini b/Templates/TP_VehicleAdv/Config/DefaultEngine.ini
--- a/Templates/TP_VehicleAdv/Config/DefaultEngine.ini
+++ b/Templates/TP_VehicleAdv/Config/DefaultEngine.ini
@@ -3,14 +3,14 @@ bSubstepping=False
 bSubsteppingAsync=False
 
 [/Script/EngineSettings.GameMapsSettings]
-EditorStartupMap=/Game/VehicleTemplate/Maps/VehicleAdvExampleMap.VehicleAdvExampleMap
+EditorStartupMap=/Game/VehicleTemplate/Maps/VehicleBasic.VehicleBasic
 LocalMapOptions=
 TransitionMap=
 bUseSplitscreen=True
 TwoPlayerSplitscreenLayout=Horizontal
 ThreePlayerSplitscreenLayout=FavorTop
 GameInstanceClass=/Script/Engine.GameInstance
-GameDefaultMap=/Game/VehicleTemplate/Maps/VehicleAdvExampleMap.VehicleAdvExampleMap
+GameDefaultMap=/Game/VehicleTemplate/Maps/VehicleBasic.VehicleBasic
 ServerDefaultMap=/Engine/Maps/Entry
 GlobalDefaultGameMode=/Game/VehicleTemplate/Blueprints/BP_VehicleAdvGameMode.BP_VehicleAdvGameMode_C
 GlobalDefaultServerGameMode=None
@@ -52,4 +52,19 @@ SoundCueCookQualityIndex=-1
 
 [/Script/LinuxTargetPlatform.LinuxTargetSettings]
 -TargetedRHIs=SF_VULKAN_SM5
-+TargetedRHIs=SF_VULKAN_SM6
\ No newline at end of file
++TargetedRHIs=SF_VULKAN_SM6
+
+[/Script/AndroidFileServerEditor.AndroidFileServerRuntimeSettings]
+bEnablePlugin=True
+bAllowNetworkConnection=True
+SecurityToken=7584171B4FAB7742F2593A9D8AE0B9F0
+bIncludeInShipping=False
+bAllowExternalStartInShipping=False
+bCompileAFSProject=False
+bUseCompression=False
+bLogFiles=False
+bReportStats=False
+ConnectionType=USBOnly
+bUseManualIPAddress=False
+ManualIPAddress=
+
diff --git a/Templates/TP_VehicleAdv/Source/TP_VehicleAdv/OffroadCar/TP_VehicleAdvOffroadCar.cpp b/Templates/TP_VehicleAdv/Source/TP_VehicleAdv/OffroadCar/TP_VehicleAdvOffroadCar.cpp
new file mode 100644
--- /dev/null
+++ b/Templates/TP_VehicleAdv/Source/TP_VehicleAdv/OffroadCar/TP_VehicleAdvOffroadCar.cpp
@@ -0,0 +1,88 @@
+// Copyright Epic Games, Inc. All Rights Reserved.
+
+
+#include "TP_VehicleAdvOffroadCar.h"
+#include "TP_VehicleAdvOffroadWheelFront.h"
+#include "TP_VehicleAdvOffroadWheelRear.h"
+#include "ChaosWheeledVehicleMovementComponent.h"
+#include "GameFramework/SpringArmComponent.h"
+#include "Components/StaticMeshComponent.h"
+#include "Components/SceneComponent.h"
+#include "Components/SkeletalMeshComponent.h"
+
+ATP_VehicleAdvOffroadCar::ATP_VehicleAdvOffroadCar()
+{
+	// construct the mesh components
+	Chassis = CreateDefaultSubobject<UStaticMeshComponent>(TEXT("Chassis"));
+	Chassis->SetupAttachment(GetMesh());
+
+	// NOTE: tire sockets are set from the Blueprint class
+	TireFrontLeft = CreateDefaultSubobject<UStaticMeshComponent>(TEXT("Tire Front Left"));
+	TireFrontLeft->SetupAttachment(GetMesh(), FName("VisWheel_FL"));
+	TireFrontLeft->SetCollisionProfileName(FName("NoCollision"));
+
+	TireFrontRight = CreateDefaultSubobject<UStaticMeshComponent>(TEXT("Tire Front Right"));
+	TireFrontRight->SetupAttachment(GetMesh(), FName("VisWheel_FR"));
+	TireFrontRight->SetCollisionProfileName(FName("NoCollision"));
+	TireFrontRight->SetRelativeRotation(FRotator(0.0f, 180.0f, 0.0f));
+
+	TireRearLeft = CreateDefaultSubobject<UStaticMeshComponent>(TEXT("Tire Rear Left"));
+	TireRearLeft->SetupAttachment(GetMesh(), FName("VisWheel_BL"));
+	TireRearLeft->SetCollisionProfileName(FName("NoCollision"));
+
+	TireRearRight = CreateDefaultSubobject<UStaticMeshComponent>(TEXT("Tire Rear Right"));
+	TireRearRight->SetupAttachment(GetMesh(), FName("VisWheel_BR"));
+	TireRearRight->SetCollisionProfileName(FName("NoCollision"));
+	TireRearRight->SetRelativeRotation(FRotator(0.0f, 180.0f, 0.0f));
+
+	// adjust the cameras
+	GetFrontSpringArm()->SetRelativeLocation(FVector(-5.0f, -30.0f, 135.0f));
+	GetBackSpringArm()->SetRelativeLocation(FVector(0.0f, 0.0f, 75.0f));
+
+	// Note: for faster iteration times, the vehicle setup can be tweaked in the Blueprint instead
+
+	// Set up the chassis
+	GetChaosVehicleMovement()->ChassisHeight = 160.0f;
+	GetChaosVehicleMovement()->DragCoefficient = 0.1f;
+	GetChaosVehicleMovement()->DownforceCoefficient = 0.1f;
+	GetChaosVehicleMovement()->CenterOfMassOverride = FVector(0.0f, 0.0f, 75.0f);
+	GetChaosVehicleMovement()->bEnableCenterOfMassOverride = true;
+
+	// Set up the wheels
+	GetChaosVehicleMovement()->bLegacyWheelFrictionPosition = false;
+	GetChaosVehicleMovement()->WheelSetups.SetNum(4);
+
+	GetChaosVehicleMovement()->WheelSetups[0].WheelClass = UTP_VehicleAdvOffroadWheelFront::StaticClass();
+	GetChaosVehicleMovement()->WheelSetups[0].BoneName = FName("PhysWheel_FL");
+	GetChaosVehicleMovement()->WheelSetups[0].AdditionalOffset = FVector(0.0f, 0.0f, 0.0f);
+
+	GetChaosVehicleMovement()->WheelSetups[1].WheelClass = UTP_VehicleAdvOffroadWheelFront::StaticClass();
+	GetChaosVehicleMovement()->WheelSetups[1].BoneName = FName("PhysWheel_FR");
+	GetChaosVehicleMovement()->WheelSetups[1].AdditionalOffset = FVector(0.0f, 0.0f, 0.0f);
+
+	GetChaosVehicleMovement()->WheelSetups[2].WheelClass = UTP_VehicleAdvOffroadWheelRear::StaticClass();
+	GetChaosVehicleMovement()->WheelSetups[2].BoneName = FName("PhysWheel_BL");
+	GetChaosVehicleMovement()->WheelSetups[2].AdditionalOffset = FVector(0.0f, 0.0f, 0.0f);
+
+	GetChaosVehicleMovement()->WheelSetups[3].WheelClass = UTP_VehicleAdvOffroadWheelRear::StaticClass();
+	GetChaosVehicleMovement()->WheelSetups[3].BoneName = FName("PhysWheel_BR");
+	GetChaosVehicleMovement()->WheelSetups[3].AdditionalOffset = FVector(0.0f, 0.0f, 0.0f);
+
+	// Set up the engine
+	// NOTE: Check the Blueprint asset for the Torque Curve
+	GetChaosVehicleMovement()->EngineSetup.MaxTorque = 600.0f;
+	GetChaosVehicleMovement()->EngineSetup.MaxRPM = 5000.0f;
+	GetChaosVehicleMovement()->EngineSetup.EngineIdleRPM = 1200.0f;
+	GetChaosVehicleMovement()->EngineSetup.EngineBrakeEffect = 0.05f;
+	GetChaosVehicleMovement()->EngineSetup.EngineRevUpMOI = 5.0f;
+	GetChaosVehicleMovement()->EngineSetup.EngineRevDownRate = 600.0f;
+
+	// Set up the differential
+	GetChaosVehicleMovement()->DifferentialSetup.DifferentialType = EVehicleDifferential::AllWheelDrive;
+	GetChaosVehicleMovement()->DifferentialSetup.FrontRearSplit = 0.5f;
+
+	// Set up the steering
+	// NOTE: Check the Blueprint asset for the Steering Curve
+	GetChaosVehicleMovement()->SteeringSetup.SteeringType = ESteeringType::AngleRatio;
+	GetChaosVehicleMovement()->SteeringSetup.AngleRatio = 0.7f;
+}
\ No newline at end of file
diff --git a/Templates/TP_VehicleAdv/Source/TP_VehicleAdv/OffroadCar/TP_VehicleAdvOffroadCar.h b/Templates/TP_VehicleAdv/Source/TP_VehicleAdv/OffroadCar/TP_VehicleAdvOffroadCar.h
new file mode 100644
--- /dev/null
+++ b/Templates/TP_VehicleAdv/Source/TP_VehicleAdv/OffroadCar/TP_VehicleAdvOffroadCar.h
@@ -0,0 +1,40 @@
+// Copyright Epic Games, Inc. All Rights Reserved.
+
+#pragma once
+
+#include "CoreMinimal.h"
+#include "TP_VehicleAdvPawn.h"
+#include "TP_VehicleAdvOffroadCar.generated.h"
+
+/**
+ *  Offroad car wheeled vehicle implementation
+ */
+UCLASS(abstract)
+class ATP_VehicleAdvOffroadCar : public ATP_VehicleAdvPawn
+{
+	GENERATED_BODY()
+	
+	/** Chassis static mesh */
+	UPROPERTY(VisibleAnywhere, BlueprintReadOnly, Category = Meshes, meta = (AllowPrivateAccess = "true"))
+	UStaticMeshComponent* Chassis;
+
+	/** FL Tire static mesh */
+	UPROPERTY(VisibleAnywhere, BlueprintReadOnly, Category = Meshes, meta = (AllowPrivateAccess = "true"))
+	UStaticMeshComponent* TireFrontLeft;
+
+	/** FR Tire static mesh */
+	UPROPERTY(VisibleAnywhere, BlueprintReadOnly, Category = Meshes, meta = (AllowPrivateAccess = "true"))
+	UStaticMeshComponent* TireFrontRight;
+
+	/** RL Tire static mesh */
+	UPROPERTY(VisibleAnywhere, BlueprintReadOnly, Category = Meshes, meta = (AllowPrivateAccess = "true"))
+	UStaticMeshComponent* TireRearLeft;
+
+	/** RR Tire static mesh */
+	UPROPERTY(VisibleAnywhere, BlueprintReadOnly, Category = Meshes, meta = (AllowPrivateAccess = "true"))
+	UStaticMeshComponent* TireRearRight;
+
+public:
+
+	ATP_VehicleAdvOffroadCar();
+};
diff --git a/Templates/TP_VehicleAdv/Source/TP_VehicleAdv/OffroadCar/TP_VehicleAdvOffroadWheelFront.cpp b/Templates/TP_VehicleAdv/Source/TP_VehicleAdv/OffroadCar/TP_VehicleAdvOffroadWheelFront.cpp
new file mode 100644
--- /dev/null
+++ b/Templates/TP_VehicleAdv/Source/TP_VehicleAdv/OffroadCar/TP_VehicleAdvOffroadWheelFront.cpp
@@ -0,0 +1,22 @@
+// Copyright Epic Games, Inc. All Rights Reserved.
+
+
+#include "TP_VehicleAdvOffroadWheelFront.h"
+
+UTP_VehicleAdvOffroadWheelFront::UTP_VehicleAdvOffroadWheelFront()
+{
+	WheelRadius = 50.0f;
+	CorneringStiffness = 750.0f;
+	FrictionForceMultiplier = 4.0f;
+	bAffectedByEngine = true;
+
+	SuspensionMaxRaise = 20.0f;
+	SuspensionMaxDrop = 20.0f;
+	WheelLoadRatio = 1.0f;
+	SpringRate = 100.0f;
+	SpringPreload = 100.0f;
+	SweepShape = ESweepShape::Shapecast;
+
+	MaxBrakeTorque = 3000.0f;
+	MaxHandBrakeTorque = 6000.0f;
+}
\ No newline at end of file
diff --git a/Templates/TP_VehicleAdv/Source/TP_VehicleAdv/OffroadCar/TP_VehicleAdvOffroadWheelFront.h b/Templates/TP_VehicleAdv/Source/TP_VehicleAdv/OffroadCar/TP_VehicleAdvOffroadWheelFront.h
new file mode 100644
--- /dev/null
+++ b/Templates/TP_VehicleAdv/Source/TP_VehicleAdv/OffroadCar/TP_VehicleAdvOffroadWheelFront.h
@@ -0,0 +1,19 @@
+// Copyright Epic Games, Inc. All Rights Reserved.
+
+#pragma once
+
+#include "CoreMinimal.h"
+#include "TP_VehicleAdvWheelFront.h"
+#include "TP_VehicleAdvOffroadWheelFront.generated.h"
+
+/**
+ *  Front wheel definition for Offroad Car.
+ */
+UCLASS()
+class UTP_VehicleAdvOffroadWheelFront : public UTP_VehicleAdvWheelFront
+{
+	GENERATED_BODY()
+	
+public:
+	UTP_VehicleAdvOffroadWheelFront();
+};
diff --git a/Templates/TP_VehicleAdv/Source/TP_VehicleAdv/OffroadCar/TP_VehicleAdvOffroadWheelRear.cpp b/Templates/TP_VehicleAdv/Source/TP_VehicleAdv/OffroadCar/TP_VehicleAdvOffroadWheelRear.cpp
new file mode 100644
--- /dev/null
+++ b/Templates/TP_VehicleAdv/Source/TP_VehicleAdv/OffroadCar/TP_VehicleAdvOffroadWheelRear.cpp
@@ -0,0 +1,21 @@
+// Copyright Epic Games, Inc. All Rights Reserved.
+
+
+#include "TP_VehicleAdvOffroadWheelRear.h"
+
+UTP_VehicleAdvOffroadWheelRear::UTP_VehicleAdvOffroadWheelRear()
+{
+	WheelRadius = 50.f;
+	CorneringStiffness = 750.0f;
+	FrictionForceMultiplier = 4.0f;
+	
+	SuspensionMaxRaise = 20.0f;
+	SuspensionMaxDrop = 20.0f;
+	WheelLoadRatio = 1.0f;
+	SpringRate = 100.0f;
+	SpringPreload = 100.0f;
+	SweepShape = ESweepShape::Shapecast;
+
+	MaxBrakeTorque = 3000.0f;
+	MaxHandBrakeTorque = 6000.0f;
+}
\ No newline at end of file
diff --git a/Templates/TP_VehicleAdv/Source/TP_VehicleAdv/OffroadCar/TP_VehicleAdvOffroadWheelRear.h b/Templates/TP_VehicleAdv/Source/TP_VehicleAdv/OffroadCar/TP_VehicleAdvOffroadWheelRear.h
new file mode 100644
--- /dev/null
+++ b/Templates/TP_VehicleAdv/Source/TP_VehicleAdv/OffroadCar/TP_VehicleAdvOffroadWheelRear.h
@@ -0,0 +1,20 @@
+// Copyright Epic Games, Inc. All Rights Reserved.
+
+#pragma once
+
+#include "CoreMinimal.h"
+#include "TP_VehicleAdvWheelRear.h"
+#include "TP_VehicleAdvOffroadWheelRear.generated.h"
+
+/**
+ *  Rear wheel definition for Offroad Car.
+ */
+UCLASS()
+class UTP_VehicleAdvOffroadWheelRear : public UTP_VehicleAdvWheelRear
+{
+	GENERATED_BODY()
+	
+public:
+
+	UTP_VehicleAdvOffroadWheelRear();
+};
diff --git a/Templates/TP_VehicleAdv/Source/TP_VehicleAdv/SportsCar/TP_VehicleAdvSportsCar.h b/Templates/TP_VehicleAdv/Source/TP_VehicleAdv/SportsCar/TP_VehicleAdvSportsCar.h
new file mode 100644
--- /dev/null
+++ b/Templates/TP_VehicleAdv/Source/TP_VehicleAdv/SportsCar/TP_VehicleAdvSportsCar.h
@@ -0,0 +1,20 @@
+// Copyright Epic Games, Inc. All Rights Reserved.
+
+#pragma once
+
+#include "CoreMinimal.h"
+#include "TP_VehicleAdvPawn.h"
+#include "TP_VehicleAdvSportsCar.generated.h"
+
+/**
+ *  Sports car wheeled vehicle implementation
+ */
+UCLASS(abstract)
+class ATP_VehicleAdvSportsCar : public ATP_VehicleAdvPawn
+{
+	GENERATED_BODY()
+	
+public:
+
+	ATP_VehicleAdvSportsCar();
+};
diff --git a/Templates/TP_VehicleAdv/Source/TP_VehicleAdv/SportsCar/TP_VehicleAdvSportsWheelFront.cpp b/Templates/TP_VehicleAdv/Source/TP_VehicleAdv/SportsCar/TP_VehicleAdvSportsWheelFront.cpp
new file mode 100644
--- /dev/null
+++ b/Templates/TP_VehicleAdv/Source/TP_VehicleAdv/SportsCar/TP_VehicleAdvSportsWheelFront.cpp
@@ -0,0 +1,14 @@
+// Copyright Epic Games, Inc. All Rights Reserved.
+
+
+#include "TP_VehicleAdvSportsWheelFront.h"
+
+UTP_VehicleAdvSportsWheelFront::UTP_VehicleAdvSportsWheelFront()
+{
+	WheelRadius = 39.0f;
+	WheelWidth = 35.0f;
+	FrictionForceMultiplier = 3.0f;
+
+	MaxBrakeTorque = 4500.0f;
+	MaxHandBrakeTorque = 6000.0f;
+}
\ No newline at end of file
diff --git a/Templates/TP_VehicleAdv/Source/TP_VehicleAdv/SportsCar/TP_VehicleAdvSportsWheelFront.h b/Templates/TP_VehicleAdv/Source/TP_VehicleAdv/SportsCar/TP_VehicleAdvSportsWheelFront.h
new file mode 100644
--- /dev/null
+++ b/Templates/TP_VehicleAdv/Source/TP_VehicleAdv/SportsCar/TP_VehicleAdvSportsWheelFront.h
@@ -0,0 +1,20 @@
+// Copyright Epic Games, Inc. All Rights Reserved.
+
+#pragma once
+
+#include "CoreMinimal.h"
+#include "TP_VehicleAdvWheelFront.h"
+#include "TP_VehicleAdvSportsWheelFront.generated.h"
+
+/**
+ *  Front wheel definition for Sports Car.
+ */
+UCLASS()
+class UTP_VehicleAdvSportsWheelFront : public UTP_VehicleAdvWheelFront
+{
+	GENERATED_BODY()
+	
+public:
+
+	UTP_VehicleAdvSportsWheelFront();
+};
diff --git a/Templates/TP_VehicleAdv/Source/TP_VehicleAdv/SportsCar/TP_VehicleAdvSportsWheelRear.cpp b/Templates/TP_VehicleAdv/Source/TP_VehicleAdv/SportsCar/TP_VehicleAdvSportsWheelRear.cpp
new file mode 100644
--- /dev/null
+++ b/Templates/TP_VehicleAdv/Source/TP_VehicleAdv/SportsCar/TP_VehicleAdvSportsWheelRear.cpp
@@ -0,0 +1,15 @@
+// Copyright Epic Games, Inc. All Rights Reserved.
+
+
+#include "TP_VehicleAdvSportsWheelRear.h"
+
+UTP_VehicleAdvSportsWheelRear::UTP_VehicleAdvSportsWheelRear()
+{
+	WheelRadius = 40.f;
+	WheelWidth = 40.0f;
+	FrictionForceMultiplier = 4.0f;
+	SlipThreshold = 100.0f;
+	SkidThreshold = 100.0f;
+	MaxSteerAngle = 0.0f;
+	MaxHandBrakeTorque = 6000.0f;
+}
\ No newline at end of file
diff --git a/Templates/TP_VehicleAdv/Source/TP_VehicleAdv/SportsCar/TP_VehicleAdvSportsWheelRear.h b/Templates/TP_VehicleAdv/Source/TP_VehicleAdv/SportsCar/TP_VehicleAdvSportsWheelRear.h
new file mode 100644
--- /dev/null
+++ b/Templates/TP_VehicleAdv/Source/TP_VehicleAdv/SportsCar/TP_VehicleAdvSportsWheelRear.h
@@ -0,0 +1,20 @@
+// Copyright Epic Games, Inc. All Rights Reserved.
+
+#pragma once
+
+#include "CoreMinimal.h"
+#include "TP_VehicleAdvWheelRear.h"
+#include "TP_VehicleAdvSportsWheelRear.generated.h"
+
+/**
+ *  Rear wheel definition for Sports Car.
+ */
+UCLASS()
+class UTP_VehicleAdvSportsWheelRear : public UTP_VehicleAdvWheelRear
+{
+	GENERATED_BODY()
+	
+public:
+
+	UTP_VehicleAdvSportsWheelRear();
+};
diff --git a/Templates/TP_VehicleAdv/Source/TP_VehicleAdv/TP_VehicleAdv.Build.cs b/Templates/TP_VehicleAdv/Source/TP_VehicleAdv/TP_VehicleAdv.Build.cs
--- a/Templates/TP_VehicleAdv/Source/TP_VehicleAdv/TP_VehicleAdv.Build.cs
+++ b/Templates/TP_VehicleAdv/Source/TP_VehicleAdv/TP_VehicleAdv.Build.cs
@@ -8,6 +8,33 @@ public class TP_VehicleAdv : ModuleRules
 	{
 		PCHUsage = PCHUsageMode.UseExplicitOrSharedPCHs;
 
-		PublicDependencyModuleNames.AddRange(new string[] { "Core", "CoreUObject", "Engine", "InputCore", "EnhancedInput", "ChaosVehicles", "PhysicsCore" });
+		PublicDependencyModuleNames.AddRange(new string[] {
+			"Core",
+			"CoreUObject",
+			"Engine",
+			"InputCore",
+			"EnhancedInput",
+			"ChaosVehicles",
+			"PhysicsCore",
+			"UMG"
+		});
+
+		PublicIncludePaths.AddRange(new string[] {
+			"TP_VehicleAdv",
+			"TP_VehicleAdv/SportsCar",
+			"TP_VehicleAdv/OffroadCar",
+			"TP_VehicleAdv/Variant_Offroad",
+			"TP_VehicleAdv/Variant_TimeTrial"
+		});
+
+		PrivateDependencyModuleNames.AddRange(new string[] { });
+
+		// Uncomment if you are using Slate UI
+		// PrivateDependencyModuleNames.AddRange(new string[] { "Slate", "SlateCore" });
+
+		// Uncomment if you are using online features
+		// PrivateDependencyModuleNames.Add("OnlineSubsystem");
+
+		// To include OnlineSubsystemSteam, add it to the plugins section in your uproject file with the Enabled attribute set to true
 	}
 }
diff --git a/Templates/TP_VehicleAdv/Source/TP_VehicleAdv/TP_VehicleAdvGameMode.h b/Templates/TP_VehicleAdv/Source/TP_VehicleAdv/TP_VehicleAdvGameMode.h
--- a/Templates/TP_VehicleAdv/Source/TP_VehicleAdv/TP_VehicleAdvGameMode.h
+++ b/Templates/TP_VehicleAdv/Source/TP_VehicleAdv/TP_VehicleAdvGameMode.h
@@ -6,7 +6,7 @@
 #include "GameFramework/GameModeBase.h"
 #include "TP_VehicleAdvGameMode.generated.h"
 
-UCLASS(MinimalAPI)
+UCLASS(abstract)
 class ATP_VehicleAdvGameMode : public AGameModeBase
 {
 	GENERATED_BODY()
diff --git a/Templates/TP_VehicleAdv/Source/TP_VehicleAdv/TP_VehicleAdvOffroadCar.cpp b/Templates/TP_VehicleAdv/Source/TP_VehicleAdv/TP_VehicleAdvOffroadCar.cpp
deleted file mode 100644
diff --git a/Templates/TP_VehicleAdv/Source/TP_VehicleAdv/TP_VehicleAdvOffroadCar.h b/Templates/TP_VehicleAdv/Source/TP_VehicleAdv/TP_VehicleAdvOffroadCar.h
deleted file mode 100644
diff --git a/Templates/TP_VehicleAdv/Source/TP_VehicleAdv/TP_VehicleAdvOffroadWheelFront.cpp b/Templates/TP_VehicleAdv/Source/TP_VehicleAdv/TP_VehicleAdvOffroadWheelFront.cpp
deleted file mode 100644
diff --git a/Templates/TP_VehicleAdv/Source/TP_VehicleAdv/TP_VehicleAdvOffroadWheelFront.h b/Templates/TP_VehicleAdv/Source/TP_VehicleAdv/TP_VehicleAdvOffroadWheelFront.h
deleted file mode 100644
diff --git a/Templates/TP_VehicleAdv/Source/TP_VehicleAdv/TP_VehicleAdvOffroadWheelRear.cpp b/Templates/TP_VehicleAdv/Source/TP_VehicleAdv/TP_VehicleAdvOffroadWheelRear.cpp
deleted file mode 100644
diff --git a/Templates/TP_VehicleAdv/Source/TP_VehicleAdv/TP_VehicleAdvOffroadWheelRear.h b/Templates/TP_VehicleAdv/Source/TP_VehicleAdv/TP_VehicleAdvOffroadWheelRear.h
deleted file mode 100644
diff --git a/Templates/TP_VehicleAdv/Source/TP_VehicleAdv/TP_VehicleAdvPlayerController.h b/Templates/TP_VehicleAdv/Source/TP_VehicleAdv/TP_VehicleAdvPlayerController.h
--- a/Templates/TP_VehicleAdv/Source/TP_VehicleAdv/TP_VehicleAdvPlayerController.h
+++ b/Templates/TP_VehicleAdv/Source/TP_VehicleAdv/TP_VehicleAdvPlayerController.h
@@ -15,31 +15,31 @@ class UTP_VehicleAdvUI;
  *  Handles input mapping and user interface
  */
 UCLASS(abstract)
-class TP_VEHICLEADV_API ATP_VehicleAdvPlayerController : public APlayerController
+class ATP_VehicleAdvPlayerController : public APlayerController
 {
 	GENERATED_BODY()
 
 protected:
 
 	/** Input Mapping Context to be used for player input */
-	UPROPERTY(EditAnywhere, BlueprintReadOnly, Category = Input)
+	UPROPERTY(EditAnywhere, BlueprintReadOnly, Category="Input")
 	UInputMappingContext* InputMappingContext;
 
 	/** If true, the optional steering wheel input mapping context will be registered */
-	UPROPERTY(EditAnywhere, BlueprintReadOnly, Category = Input)
+	UPROPERTY(EditAnywhere, BlueprintReadOnly, Category = "Input")
 	bool bUseSteeringWheelControls = false;
 
 	/** Optional Input Mapping Context to be used for steering wheel input.
 	 *  This is added alongside the default Input Mapping Context and does not block other forms of input.
 	 */
-	UPROPERTY(EditAnywhere, BlueprintReadOnly, Category = Input, meta=(EditCondition="bUseSteeringWheelControls"))
+	UPROPERTY(EditAnywhere, BlueprintReadOnly, Category = "Input", meta=(EditCondition = "bUseSteeringWheelControls"))
 	UInputMappingContext* SteeringWheelInputMappingContext;
 
 	/** Pointer to the controlled vehicle pawn */
 	TObjectPtr<ATP_VehicleAdvPawn> VehiclePawn;
 
 	/** Type of the UI to spawn */
-	UPROPERTY(EditAnywhere, BlueprintReadOnly, Category = UI)
+	UPROPERTY(EditAnywhere, BlueprintReadOnly, Category="UI")
 	TSubclassOf<UTP_VehicleAdvUI> VehicleUIClass;
 
 	/** Pointer to the UI widget */
diff --git a/Templates/TP_VehicleAdv/Source/TP_VehicleAdv/TP_VehicleAdvSportsCar.h b/Templates/TP_VehicleAdv/Source/TP_VehicleAdv/TP_VehicleAdvSportsCar.h
deleted file mode 100644
diff --git a/Templates/TP_VehicleAdv/Source/TP_VehicleAdv/TP_VehicleAdvSportsWheelFront.cpp b/Templates/TP_VehicleAdv/Source/TP_VehicleAdv/TP_VehicleAdvSportsWheelFront.cpp
deleted file mode 100644
diff --git a/Templates/TP_VehicleAdv/Source/TP_VehicleAdv/TP_VehicleAdvSportsWheelFront.h b/Templates/TP_VehicleAdv/Source/TP_VehicleAdv/TP_VehicleAdvSportsWheelFront.h
deleted file mode 100644
diff --git a/Templates/TP_VehicleAdv/Source/TP_VehicleAdv/TP_VehicleAdvSportsWheelRear.cpp b/Templates/TP_VehicleAdv/Source/TP_VehicleAdv/TP_VehicleAdvSportsWheelRear.cpp
deleted file mode 100644
diff --git a/Templates/TP_VehicleAdv/Source/TP_VehicleAdv/TP_VehicleAdvSportsWheelRear.h b/Templates/TP_VehicleAdv/Source/TP_VehicleAdv/TP_VehicleAdvSportsWheelRear.h
deleted file mode 100644
diff --git a/Templates/TP_VehicleAdv/Source/TP_VehicleAdv/TP_VehicleAdvUI.h b/Templates/TP_VehicleAdv/Source/TP_VehicleAdv/TP_VehicleAdvUI.h
--- a/Templates/TP_VehicleAdv/Source/TP_VehicleAdv/TP_VehicleAdvUI.h
+++ b/Templates/TP_VehicleAdv/Source/TP_VehicleAdv/TP_VehicleAdvUI.h
@@ -12,7 +12,7 @@
  *  Widget setup is handled in a Blueprint subclass.
  */
 UCLASS(abstract)
-class TP_VEHICLEADV_API UTP_VehicleAdvUI : public UUserWidget
+class UTP_VehicleAdvUI : public UUserWidget
 {
 	GENERATED_BODY()
 	
diff --git a/Templates/TP_VehicleAdv/Source/TP_VehicleAdv/Variant_OffRoad/OffroadGameMode.cpp b/Templates/TP_VehicleAdv/Source/TP_VehicleAdv/Variant_OffRoad/OffroadGameMode.cpp
new file mode 100644
--- /dev/null
+++ b/Templates/TP_VehicleAdv/Source/TP_VehicleAdv/Variant_OffRoad/OffroadGameMode.cpp
@@ -0,0 +1,9 @@
+// Copyright Epic Games, Inc. All Rights Reserved.
+
+
+#include "Variant_OffRoad/OffroadGameMode.h"
+
+AOffroadGameMode::AOffroadGameMode()
+{
+
+}
\ No newline at end of file
diff --git a/Templates/TP_VehicleAdv/Source/TP_VehicleAdv/Variant_OffRoad/OffroadGameMode.h b/Templates/TP_VehicleAdv/Source/TP_VehicleAdv/Variant_OffRoad/OffroadGameMode.h
new file mode 100644
--- /dev/null
+++ b/Templates/TP_VehicleAdv/Source/TP_VehicleAdv/Variant_OffRoad/OffroadGameMode.h
@@ -0,0 +1,21 @@
+// Copyright Epic Games, Inc. All Rights Reserved.
+
+#pragma once
+
+#include "CoreMinimal.h"
+#include "GameFramework/GameModeBase.h"
+#include "OffroadGameMode.generated.h"
+
+/**
+ *  Simple GameMode for an offroad vehicle game
+ */
+UCLASS(abstract)
+class AOffroadGameMode : public AGameModeBase
+{
+	GENERATED_BODY()
+	
+public:
+
+	/** Constructor */
+	AOffroadGameMode();
+};
diff --git a/Templates/TP_VehicleAdv/Source/TP_VehicleAdv/Variant_TimeTrial/TimeTrialGameMode.cpp b/Templates/TP_VehicleAdv/Source/TP_VehicleAdv/Variant_TimeTrial/TimeTrialGameMode.cpp
new file mode 100644
--- /dev/null
+++ b/Templates/TP_VehicleAdv/Source/TP_VehicleAdv/Variant_TimeTrial/TimeTrialGameMode.cpp
@@ -0,0 +1,29 @@
+// Copyright Epic Games, Inc. All Rights Reserved.
+
+
+#include "TimeTrialGameMode.h"
+#include "Kismet/GameplayStatics.h"
+#include "TimeTrialTrackGate.h"
+#include "Engine/World.h"
+
+void ATimeTrialGameMode::BeginPlay()
+{
+	Super::BeginPlay();
+
+	// get the finish line marker
+	TArray<AActor*> ActorList;
+
+	UGameplayStatics::GetAllActorsOfClassWithTag(GetWorld(), ATimeTrialTrackGate::StaticClass(), FinishTag, ActorList);
+
+	if (ActorList.Num() > 0)
+	{
+		// get the first returned track marker that matches the tag
+		FinishLineMarker = Cast<ATimeTrialTrackGate>(ActorList[0]);
+	}
+
+}
+
+ATimeTrialTrackGate* ATimeTrialGameMode::GetFinishLine() const
+{
+	return FinishLineMarker;
+}
diff --git a/Templates/TP_VehicleAdv/Source/TP_VehicleAdv/Variant_TimeTrial/TimeTrialGameMode.h b/Templates/TP_VehicleAdv/Source/TP_VehicleAdv/Variant_TimeTrial/TimeTrialGameMode.h
new file mode 100644
--- /dev/null
+++ b/Templates/TP_VehicleAdv/Source/TP_VehicleAdv/Variant_TimeTrial/TimeTrialGameMode.h
@@ -0,0 +1,45 @@
+// Copyright Epic Games, Inc. All Rights Reserved.
+
+#pragma once
+
+#include "CoreMinimal.h"
+#include "GameFramework/GameModeBase.h"
+#include "TimeTrialGameMode.generated.h"
+
+class ATimeTrialTrackGate;
+
+/**
+ *  A simple GameMode for a Time Trial racing game
+ */
+UCLASS(abstract)
+class ATimeTrialGameMode : public AGameModeBase
+{
+	GENERATED_BODY()
+	
+protected:
+
+	/** Actor tag used to find the finish line marker on the level */
+	UPROPERTY(EditAnywhere, Category="Time Trial")
+	FName FinishTag;
+
+	/** Number of laps for the race */
+	UPROPERTY(EditAnywhere, Category="Time Trial")
+	int32 Laps = 3;
+
+	/** Pointer to the finish line track marker */
+	TObjectPtr<ATimeTrialTrackGate> FinishLineMarker;
+
+protected:
+
+	/** Gameplay initialization */
+	virtual void BeginPlay() override;
+
+public: 
+
+	/** Returns the track marker for the finish line */
+	ATimeTrialTrackGate* GetFinishLine() const;
+
+	/** Returns the number of laps for the race */
+	int32 GetLaps() const { return Laps; };
+
+};
diff --git a/Templates/TP_VehicleAdv/Source/TP_VehicleAdv/Variant_TimeTrial/TimeTrialPlayerController.cpp b/Templates/TP_VehicleAdv/Source/TP_VehicleAdv/Variant_TimeTrial/TimeTrialPlayerController.cpp
new file mode 100644
--- /dev/null
+++ b/Templates/TP_VehicleAdv/Source/TP_VehicleAdv/Variant_TimeTrial/TimeTrialPlayerController.cpp
@@ -0,0 +1,89 @@
+// Copyright Epic Games, Inc. All Rights Reserved.
+
+
+#include "TimeTrialPlayerController.h"
+#include "TimeTrialUI.h"
+#include "Engine/World.h"
+#include "TimeTrialGameMode.h"
+#include "TimeTrialTrackGate.h"
+#include "EnhancedInputSubsystems.h"
+#include "Engine/LocalPlayer.h"
+#include "InputMappingContext.h"
+#include "TP_VehicleAdvUI.h"
+
+void ATimeTrialPlayerController::BeginPlay()
+{
+	Super::BeginPlay();
+
+	// create the UI widget
+	UIWidget = CreateWidget<UTimeTrialUI>(this, UIWidgetClass);
+	UIWidget->AddToViewport(0);
+
+	// spawn the UI widget and add it to the viewport
+	VehicleUI = CreateWidget<UTP_VehicleAdvUI>(this, VehicleUIClass);
+	VehicleUI->AddToViewport();
+
+	// subscribe to the race start delegate
+	UIWidget->OnRaceStart.AddDynamic(this, &ATimeTrialPlayerController::StartRace);
+}
+
+void ATimeTrialPlayerController::SetupInputComponent()
+{
+	Super::SetupInputComponent();
+
+	// get the enhanced input subsystem
+	if (UEnhancedInputLocalPlayerSubsystem* Subsystem = ULocalPlayer::GetSubsystem<UEnhancedInputLocalPlayerSubsystem>(GetLocalPlayer()))
+	{
+		// add the mapping context so we get controls
+		Subsystem->AddMappingContext(InputMappingContext, 0);
+
+		// optionally add the steering wheel context
+		if (bUseSteeringWheelControls && SteeringWheelInputMappingContext)
+		{
+			Subsystem->AddMappingContext(SteeringWheelInputMappingContext, 1);
+		}
+	}
+}
+
+void ATimeTrialPlayerController::OnPossess(APawn* aPawn)
+{
+	Super::OnPossess(aPawn);
+
+	// disable input on the pawn
+	aPawn->DisableInput(this);
+}
+
+void ATimeTrialPlayerController::StartRace()
+{
+	// get the finish line from the game mode
+	if (ATimeTrialGameMode* GM = Cast<ATimeTrialGameMode>(GetWorld()->GetAuthGameMode()))
+	{
+		SetTargetGate(GM->GetFinishLine()->GetNextMarker());
+	}
+
+	// start the first lap
+	CurrentLap = 0;
+	IncrementLapCount();
+
+	// enable input on the pawn
+	GetPawn()->EnableInput(this);
+}
+
+void ATimeTrialPlayerController::IncrementLapCount()
+{
+	// increment the lap counter
+	++CurrentLap;
+
+	// update the UI
+	UIWidget->UpdateLapCount(CurrentLap, GetWorld()->GetTimeSeconds());
+}
+
+ATimeTrialTrackGate* ATimeTrialPlayerController::GetTargetGate()
+{
+	return TargetGate.Get();
+}
+
+void ATimeTrialPlayerController::SetTargetGate(ATimeTrialTrackGate* Gate)
+{
+	TargetGate = Gate;
+}
diff --git a/Templates/TP_VehicleAdv/Source/TP_VehicleAdv/Variant_TimeTrial/TimeTrialPlayerController.h b/Templates/TP_VehicleAdv/Source/TP_VehicleAdv/Variant_TimeTrial/TimeTrialPlayerController.h
new file mode 100644
--- /dev/null
+++ b/Templates/TP_VehicleAdv/Source/TP_VehicleAdv/Variant_TimeTrial/TimeTrialPlayerController.h
@@ -0,0 +1,83 @@
+// Copyright Epic Games, Inc. All Rights Reserved.
+
+#pragma once
+
+#include "CoreMinimal.h"
+#include "GameFramework/PlayerController.h"
+#include "TimeTrialPlayerController.generated.h"
+
+class ATimeTrialTrackGate;
+class UTimeTrialUI;
+class UInputMappingContext;
+class UTP_VehicleAdvUI;
+
+/**
+ *  A simple PlayerController for a Time Trial racing game
+ */
+UCLASS(abstract)
+class ATimeTrialPlayerController : public APlayerController
+{
+	GENERATED_BODY()
+	
+protected:
+
+	/** Input Mapping Context to be used for player input */
+	UPROPERTY(EditAnywhere, BlueprintReadOnly, Category="Input")
+	UInputMappingContext* InputMappingContext;
+
+	/** If true, the optional steering wheel input mapping context will be registered */
+	UPROPERTY(EditAnywhere, BlueprintReadOnly, Category = "Input")
+	bool bUseSteeringWheelControls = false;
+
+	/** Optional Input Mapping Context to be used for steering wheel input.
+	 *  This is added alongside the default Input Mapping Context and does not block other forms of input.
+	 */
+	UPROPERTY(EditAnywhere, BlueprintReadOnly, Category = "Input", meta=(EditCondition = "bUseSteeringWheelControls"))
+	UInputMappingContext* SteeringWheelInputMappingContext;
+
+	/** Type of UI widget to spawn*/
+	UPROPERTY(EditAnywhere, Category="Time Trial")
+	TSubclassOf<UTimeTrialUI> UIWidgetClass;
+
+	/** Pointer to the UI Widget */
+	TObjectPtr<UTimeTrialUI> UIWidget;
+
+	/** Type of the UI to spawn */
+	UPROPERTY(EditAnywhere, BlueprintReadOnly, Category="UI")
+	TSubclassOf<UTP_VehicleAdvUI> VehicleUIClass;
+
+	/** Pointer to the UI widget */
+	TObjectPtr<UTP_VehicleAdvUI> VehicleUI;
+
+	/** Next track gate the car should pass */
+	TObjectPtr<ATimeTrialTrackGate> TargetGate;
+
+	/** Lap counter */
+	int32 CurrentLap = 0;
+
+protected:
+
+	/** Gameplay initialization */
+	virtual void BeginPlay() override;
+
+	/** Input initialization */
+	virtual void SetupInputComponent() override;
+
+	/** Pawn initialization */
+	virtual void OnPossess(APawn* aPawn) override;
+
+public:
+
+	/** Sets up the race start */
+	UFUNCTION()
+	void StartRace();
+
+	/** Moves on to the next lap */
+	void IncrementLapCount();
+
+	/** Returns the current target track gate */
+	ATimeTrialTrackGate* GetTargetGate();
+
+	/** Sets the target track gate for this player */
+	void SetTargetGate(ATimeTrialTrackGate* Gate);
+};
diff --git a/Templates/TP_VehicleAdv/Source/TP_VehicleAdv/Variant_TimeTrial/TimeTrialStartUI.cpp b/Templates/TP_VehicleAdv/Source/TP_VehicleAdv/Variant_TimeTrial/TimeTrialStartUI.cpp
new file mode 100644
--- /dev/null
+++ b/Templates/TP_VehicleAdv/Source/TP_VehicleAdv/Variant_TimeTrial/TimeTrialStartUI.cpp
@@ -0,0 +1,16 @@
+// Copyright Epic Games, Inc. All Rights Reserved.
+
+
+#include "TimeTrialStartUI.h"
+
+void UTimeTrialStartUI::StartCountdown()
+{
+	// pass control to BP
+	BP_StartCountdown();
+}
+
+void UTimeTrialStartUI::FinishCountdown()
+{
+	// broadcast the delegate
+	OnCountdownFinished.Broadcast();
+}
diff --git a/Templates/TP_VehicleAdv/Source/TP_VehicleAdv/Variant_TimeTrial/TimeTrialStartUI.h b/Templates/TP_VehicleAdv/Source/TP_VehicleAdv/Variant_TimeTrial/TimeTrialStartUI.h
new file mode 100644
--- /dev/null
+++ b/Templates/TP_VehicleAdv/Source/TP_VehicleAdv/Variant_TimeTrial/TimeTrialStartUI.h
@@ -0,0 +1,40 @@
+// Copyright Epic Games, Inc. All Rights Reserved.
+
+#pragma once
+
+#include "CoreMinimal.h"
+#include "Blueprint/UserWidget.h"
+#include "TimeTrialStartUI.generated.h"
+
+DECLARE_DYNAMIC_MULTICAST_DELEGATE(FCountdownFinishedDelegate);
+
+/**
+ *  A race start countdown widget.
+ *  The countdown animation is performed by widget animation.
+ *  Calls a delegate when the countdown is done to start the race.
+ */
+UCLASS(abstract)
+class UTimeTrialStartUI : public UUserWidget
+{
+	GENERATED_BODY()
+	
+public:
+
+	/** Starts the race countdown */
+	void StartCountdown();
+
+protected:
+
+	/** Passes control to Blueprint to animate the race countdown. FinishCountdown should be called to start the race when it's done. */
+	UFUNCTION(BlueprintImplementableEvent, Category="Countdown", meta=(DisplayName = "Start Countdown"))
+	void BP_StartCountdown();
+
+	/** Finishes the countdown and starts the race. */
+	UFUNCTION(BlueprintCallable, Category="Countdown")
+	void FinishCountdown();
+
+public:
+
+	FCountdownFinishedDelegate OnCountdownFinished;
+
+};
diff --git a/Templates/TP_VehicleAdv/Source/TP_VehicleAdv/Variant_TimeTrial/TimeTrialTrackGate.cpp b/Templates/TP_VehicleAdv/Source/TP_VehicleAdv/Variant_TimeTrial/TimeTrialTrackGate.cpp
new file mode 100644
--- /dev/null
+++ b/Templates/TP_VehicleAdv/Source/TP_VehicleAdv/Variant_TimeTrial/TimeTrialTrackGate.cpp
@@ -0,0 +1,50 @@
+// Copyright Epic Games, Inc. All Rights Reserved.
+
+
+#include "TimeTrialTrackGate.h"
+#include "Components/SceneComponent.h"
+#include "Components/BoxComponent.h"
+#include "TimeTrialPlayerController.h"
+
+ATimeTrialTrackGate::ATimeTrialTrackGate()
+{
+ 	PrimaryActorTick.bCanEverTick = true;
+
+	// create the root component
+	RootComponent = CreateDefaultSubobject<USceneComponent>(TEXT("Root"));
+
+	// create the collision box
+	CollisionBox = CreateDefaultSubobject<UBoxComponent>(TEXT("Collision Box"));
+	CollisionBox->SetupAttachment(RootComponent);
+
+	CollisionBox->SetBoxExtent(FVector(1000.0f));
+	CollisionBox->SetLineThickness(32.0f);
+	CollisionBox->bHiddenInGame = false;
+	CollisionBox->SetCollisionProfileName(FName("OverlapAllDynamic"));
+
+}
+
+void ATimeTrialTrackGate::NotifyActorBeginOverlap(AActor* OtherActor)
+{
+	// get the player controller of the overlapping actor
+	if (ATimeTrialPlayerController* PC = Cast<ATimeTrialPlayerController>(OtherActor->GetInstigatorController()))
+	{
+		// is this the current target marker for the player?
+		if (PC->GetTargetGate() == this)
+		{
+			// point the player to the next marker
+			PC->SetTargetGate(NextMarker);
+
+			// if this is the finish line, increment the lap
+			if (bIsFinishLine)
+			{
+				PC->IncrementLapCount();
+			}
+		}
+	}
+}
+
+ATimeTrialTrackGate* ATimeTrialTrackGate::GetNextMarker() const
+{
+	return NextMarker;
+}
diff --git a/Templates/TP_VehicleAdv/Source/TP_VehicleAdv/Variant_TimeTrial/TimeTrialTrackGate.h b/Templates/TP_VehicleAdv/Source/TP_VehicleAdv/Variant_TimeTrial/TimeTrialTrackGate.h
new file mode 100644
--- /dev/null
+++ b/Templates/TP_VehicleAdv/Source/TP_VehicleAdv/Variant_TimeTrial/TimeTrialTrackGate.h
@@ -0,0 +1,48 @@
+// Copyright Epic Games, Inc. All Rights Reserved.
+
+#pragma once
+
+#include "CoreMinimal.h"
+#include "GameFramework/Actor.h"
+#include "TimeTrialTrackGate.generated.h"
+
+class UBoxComponent;
+
+/**
+ *  A track gate volume for a Time Trial racing game.
+ *  Players must pass through the track gates in order to complete a lap.
+ */
+UCLASS(abstract)
+class ATimeTrialTrackGate : public AActor
+{
+	GENERATED_BODY()
+	
+	/** Collision Box */
+	UPROPERTY(VisibleAnywhere, BlueprintReadOnly, Category = "Components", meta = (AllowPrivateAccess = "true"))
+	UBoxComponent* CollisionBox;
+
+protected:
+
+	/** If this is set to true, this track gate is considered the finish line and will increase the lap when passed */
+	UPROPERTY(EditAnywhere, Category="Track Gate")
+	bool bIsFinishLine = false;
+
+	/** Pointer to the next track marker in the sequence */
+	UPROPERTY(EditAnywhere, Category="Track Gate")
+	ATimeTrialTrackGate* NextMarker;
+
+public:	
+	
+	/** Constructor */
+	ATimeTrialTrackGate();
+
+protected:
+
+	/** Handle collision */
+	virtual void NotifyActorBeginOverlap(AActor* OtherActor) override;
+
+public:
+
+	/** Returns the next marker on the track */
+	ATimeTrialTrackGate* GetNextMarker() const;
+};
diff --git a/Templates/TP_VehicleAdv/Source/TP_VehicleAdv/Variant_TimeTrial/TimeTrialUI.cpp b/Templates/TP_VehicleAdv/Source/TP_VehicleAdv/Variant_TimeTrial/TimeTrialUI.cpp
new file mode 100644
--- /dev/null
+++ b/Templates/TP_VehicleAdv/Source/TP_VehicleAdv/Variant_TimeTrial/TimeTrialUI.cpp
@@ -0,0 +1,68 @@
+// Copyright Epic Games, Inc. All Rights Reserved.
+
+
+#include "TimeTrialUI.h"
+#include "TimeTrialStartUI.h"
+
+void UTimeTrialUI::NativeConstruct()
+{
+	Super::NativeConstruct();
+
+	// create the start countdown widget
+	UTimeTrialStartUI* StartUI = CreateWidget<UTimeTrialStartUI>(GetOwningPlayer(), StartUIClass);
+	StartUI->AddToViewport(0);
+
+	// subscribe to the countdown finished delegate
+	StartUI->OnCountdownFinished.AddDynamic(this, &UTimeTrialUI::StartRace);
+
+	// start the countdown
+	StartUI->StartCountdown();
+}
+
+void UTimeTrialUI::UpdateLapCount(int32 Lap, float LapStartTime)
+{
+	// calculate the lap time
+	const float LapTime = LapStartTime - LastLapTime;
+
+	// is this the first lap?
+	if (Lap > 1)
+	{
+		// do we have an invalid lap time?
+		if (BestLapTime < 0.0f)
+		{
+			// save the current lap time
+			BestLapTime = LapTime;
+
+		} else {
+
+			// not the first lap: do we have a lower lap time?
+			if (LapTime < BestLapTime)
+			{
+				// save the best lap time
+				BestLapTime = LapTime;
+			}
+
+		}
+		
+	} else {
+
+		// first lap: save an invalid lap time
+		BestLapTime = -1.0f;
+
+	}
+
+	// save the current lap
+	CurrentLap = Lap;
+
+	// save the lap start time
+	LastLapTime = LapStartTime;
+
+	// pass control to BP to update the widgets
+	BP_UpdateLaps();
+}
+
+void UTimeTrialUI::StartRace()
+{
+	// broadcast the delegate
+	OnRaceStart.Broadcast();
+}
diff --git a/Templates/TP_VehicleAdv/Source/TP_VehicleAdv/Variant_TimeTrial/TimeTrialUI.h b/Templates/TP_VehicleAdv/Source/TP_VehicleAdv/Variant_TimeTrial/TimeTrialUI.h
new file mode 100644
--- /dev/null
+++ b/Templates/TP_VehicleAdv/Source/TP_VehicleAdv/Variant_TimeTrial/TimeTrialUI.h
@@ -0,0 +1,72 @@
+// Copyright Epic Games, Inc. All Rights Reserved.
+
+#pragma once
+
+#include "CoreMinimal.h"
+#include "Blueprint/UserWidget.h"
+#include "TimeTrialUI.generated.h"
+
+class UTimeTrialStartUI;
+
+DECLARE_DYNAMIC_MULTICAST_DELEGATE(FStartRaceDelegate);
+
+/**
+ *  Simple UI for a Time Trial racing game
+ *  Keeps track of lap number and best time
+ *  Spawns a sub-widget to do the initial countdown
+ */
+UCLASS(abstract)
+class UTimeTrialUI : public UUserWidget
+{
+	GENERATED_BODY()
+	
+protected:
+
+	/** Type of start countdown UI widget to spawn */
+	UPROPERTY(EditAnywhere, Category="Start Countdown")
+	TSubclassOf<UTimeTrialStartUI> StartUIClass;
+
+	/** Time when the previous lap started, in seconds */
+	float LastLapTime = 0.0f;
+
+	/** Best lap time, in seconds */
+	float BestLapTime = 0.0f;
+
+	/** Current lap number */
+	int32 CurrentLap = 0;
+
+public:
+
+	/** Delegate to broadcast when the race starts */
+	FStartRaceDelegate OnRaceStart;
+
+protected:
+
+	/** Widget initialization */
+	virtual void NativeConstruct() override;
+
+public:
+
+	/** Increments the lap and updates the lap counter */
+	void UpdateLapCount(int32 Lap, float LapStartTime);
+
+	/** Allows Blueprint control to update the lap tracker widgets */
+	UFUNCTION(BlueprintImplementableEvent, Category="Time Trial", meta=(DisplayName = "Update Laps"))
+	void BP_UpdateLaps();
+
+protected:
+
+	/** Called from the countdown delegate to start the race */
+	UFUNCTION()
+	void StartRace();
+
+	/** Gets the current lap number */
+	UFUNCTION(BlueprintPure, Category="Time Trial")
+	int32 GetCurrentLap() const { return CurrentLap; };
+
+	/** Gets the best lap time saved */
+	UFUNCTION(BlueprintPure, Category="Time Trial")
+	float GetBestLapTime() const { return BestLapTime; };
+
+
+};
