Templates TP_VehicleAdv - Exposing input calls to BP for touch interface support
#rnx

diff --git a/Templates/TP_VehicleAdv/Source/TP_VehicleAdv/TP_VehicleAdvPawn.cpp b/Templates/TP_VehicleAdv/Source/TP_VehicleAdv/TP_VehicleAdvPawn.cpp
--- a/Templates/TP_VehicleAdv/Source/TP_VehicleAdv/TP_VehicleAdvPawn.cpp
+++ b/Templates/TP_VehicleAdv/Source/TP_VehicleAdv/TP_VehicleAdvPawn.cpp
@@ -109,47 +109,98 @@ void ATP_VehicleAdvPawn::Tick(float Delta)
 
 void ATP_VehicleAdvPawn::Steering(const FInputActionValue& Value)
 {
-	// get the input magnitude for steering
-	float SteeringValue = Value.Get<float>();
+	// route the input
+	DoSteering(Value.Get<float>());
+}
+
+void ATP_VehicleAdvPawn::Throttle(const FInputActionValue& Value)
+{
+	// route the input
+	DoThrottle(Value.Get<float>());
+}
+
+void ATP_VehicleAdvPawn::Brake(const FInputActionValue& Value)
+{
+	// route the input
+	DoBrake(Value.Get<float>());
+}
+
+void ATP_VehicleAdvPawn::StartBrake(const FInputActionValue& Value)
+{
+	// route the input
+	DoBrakeStart();
+}
+
+void ATP_VehicleAdvPawn::StopBrake(const FInputActionValue& Value)
+{
+	// route the input
+	DoBrakeStop();
+}
+
+void ATP_VehicleAdvPawn::StartHandbrake(const FInputActionValue& Value)
+{
+	// route the input
+	DoHandbrakeStart();
+}
+
+void ATP_VehicleAdvPawn::StopHandbrake(const FInputActionValue& Value)
+{
+	// route the input
+	DoHandbrakeStop();
+}
+
+void ATP_VehicleAdvPawn::LookAround(const FInputActionValue& Value)
+{
+	// route the input
+	DoLookAround(Value.Get<float>());
+}
+
+void ATP_VehicleAdvPawn::ToggleCamera(const FInputActionValue& Value)
+{
+	// route the input
+	DoToggleCamera();
+}
+
+void ATP_VehicleAdvPawn::ResetVehicle(const FInputActionValue& Value)
+{
+	// route the input
+	DoResetVehicle();
+}
 
+void ATP_VehicleAdvPawn::DoSteering(float SteeringValue)
+{
 	// add the input
 	ChaosVehicleMovement->SetSteeringInput(SteeringValue);
 }
 
-void ATP_VehicleAdvPawn::Throttle(const FInputActionValue& Value)
+void ATP_VehicleAdvPawn::DoThrottle(float ThrottleValue)
 {
-	// get the input magnitude for the throttle
-	float ThrottleValue = Value.Get<float>();
-
 	// add the input
 	ChaosVehicleMovement->SetThrottleInput(ThrottleValue);
 }
 
-void ATP_VehicleAdvPawn::Brake(const FInputActionValue& Value)
+void ATP_VehicleAdvPawn::DoBrake(float BrakeValue)
 {
-	// get the input magnitude for the brakes
-	float BreakValue = Value.Get<float>();
-
 	// add the input
-	ChaosVehicleMovement->SetBrakeInput(BreakValue);
+	ChaosVehicleMovement->SetBrakeInput(BrakeValue);
 }
 
-void ATP_VehicleAdvPawn::StartBrake(const FInputActionValue& Value)
+void ATP_VehicleAdvPawn::DoBrakeStart()
 {
-	// call the Blueprint hook for the break lights
+	// call the Blueprint hook for the brake lights
 	BrakeLights(true);
 }
 
-void ATP_VehicleAdvPawn::StopBrake(const FInputActionValue& Value)
+void ATP_VehicleAdvPawn::DoBrakeStop()
 {
-	// call the Blueprint hook for the break lights
+	// call the Blueprint hook for the brake lights
 	BrakeLights(false);
 
 	// reset brake input to zero
 	ChaosVehicleMovement->SetBrakeInput(0.0f);
 }
 
-void ATP_VehicleAdvPawn::StartHandbrake(const FInputActionValue& Value)
+void ATP_VehicleAdvPawn::DoHandbrakeStart()
 {
 	// add the input
 	ChaosVehicleMovement->SetHandbrakeInput(true);
@@ -158,7 +209,7 @@ void ATP_VehicleAdvPawn::StartHandbrake(const FInputActionValue& Value)
 	BrakeLights(true);
 }
 
-void ATP_VehicleAdvPawn::StopHandbrake(const FInputActionValue& Value)
+void ATP_VehicleAdvPawn::DoHandbrakeStop()
 {
 	// add the input
 	ChaosVehicleMovement->SetHandbrakeInput(false);
@@ -167,16 +218,13 @@ void ATP_VehicleAdvPawn::StopHandbrake(const FInputActionValue& Value)
 	BrakeLights(false);
 }
 
-void ATP_VehicleAdvPawn::LookAround(const FInputActionValue& Value)
+void ATP_VehicleAdvPawn::DoLookAround(float YawDelta)
 {
-	// get the flat angle value for the input 
-	float LookValue = Value.Get<float>();
-
-	// add the input
-	BackSpringArm->AddLocalRotation(FRotator(0.0f, LookValue, 0.0f));
+	// rotate the spring arm
+	BackSpringArm->AddLocalRotation(FRotator(0.0f, YawDelta, 0.0f));
 }
 
-void ATP_VehicleAdvPawn::ToggleCamera(const FInputActionValue& Value)
+void ATP_VehicleAdvPawn::DoToggleCamera()
 {
 	// toggle the active camera flag
 	bFrontCameraActive = !bFrontCameraActive;
@@ -185,7 +233,7 @@ void ATP_VehicleAdvPawn::ToggleCamera(const FInputActionValue& Value)
 	BackCamera->SetActive(!bFrontCameraActive);
 }
 
-void ATP_VehicleAdvPawn::ResetVehicle(const FInputActionValue& Value)
+void ATP_VehicleAdvPawn::DoResetVehicle()
 {
 	// reset to a location slightly above our current one
 	FVector ResetLocation = GetActorLocation() + FVector(0.0f, 0.0f, 50.0f);
@@ -194,7 +242,7 @@ void ATP_VehicleAdvPawn::ResetVehicle(const FInputActionValue& Value)
 	FRotator ResetRotation = GetActorRotation();
 	ResetRotation.Pitch = 0.0f;
 	ResetRotation.Roll = 0.0f;
-	
+
 	// teleport the actor to the reset spot and reset physics
 	SetActorTransform(FTransform(ResetRotation, ResetLocation, FVector::OneVector), false, nullptr, ETeleportType::TeleportPhysics);
 
diff --git a/Templates/TP_VehicleAdv/Source/TP_VehicleAdv/TP_VehicleAdvPawn.h b/Templates/TP_VehicleAdv/Source/TP_VehicleAdv/TP_VehicleAdvPawn.h
--- a/Templates/TP_VehicleAdv/Source/TP_VehicleAdv/TP_VehicleAdvPawn.h
+++ b/Templates/TP_VehicleAdv/Source/TP_VehicleAdv/TP_VehicleAdvPawn.h
@@ -121,6 +121,50 @@ protected:
 	/** Handles reset vehicle input */
 	void ResetVehicle(const FInputActionValue& Value);
 
+public:
+
+	/** Handle steering input by input actions or mobile interface */
+	UFUNCTION(BlueprintCallable, Category="Input")
+	void DoSteering(float SteeringValue);
+
+	/** Handle throttle input by input actions or mobile interface */
+	UFUNCTION(BlueprintCallable, Category="Input")
+	void DoThrottle(float ThrottleValue);
+
+	/** Handle brake input by input actions or mobile interface */
+	UFUNCTION(BlueprintCallable, Category="Input")
+	void DoBrake(float BrakeValue);
+
+	/** Handle brake start input by input actions or mobile interface */
+	UFUNCTION(BlueprintCallable, Category="Input")
+	void DoBrakeStart();
+
+	/** Handle brake stop input by input actions or mobile interface */
+	UFUNCTION(BlueprintCallable, Category="Input")
+	void DoBrakeStop();
+
+	/** Handle handbrake start input by input actions or mobile interface */
+	UFUNCTION(BlueprintCallable, Category="Input")
+	void DoHandbrakeStart();
+
+	/** Handle handbrake stop input by input actions or mobile interface */
+	UFUNCTION(BlueprintCallable, Category="Input")
+	void DoHandbrakeStop();
+
+	/** Handle look input by input actions or mobile interface */
+	UFUNCTION(BlueprintCallable, Category="Input")
+	void DoLookAround(float YawDelta);
+
+	/** Handle toggle camera input by input actions or mobile interface */
+	UFUNCTION(BlueprintCallable, Category="Input")
+	void DoToggleCamera();
+
+	/** Handle reset vehicle input by input actions or mobile interface */
+	UFUNCTION(BlueprintCallable, Category="Input")
+	void DoResetVehicle();
+
+protected:
+
 	/** Called when the brake lights are turned on or off */
 	UFUNCTION(BlueprintImplementableEvent, Category="Vehicle")
 	void BrakeLights(bool bBraking);
