Templates - Bug fixes
TP_ThirdPerson - Horror character can no longer fall off ledges
TP_ThirdPerson - Adjusted jump velocity and gravity scaling on default character
TP_ThirdPerson - Adjusted combat StateTree color scheme
TP_ThirdPersonBP - Adjusted jump velocity and gravity scaling on default character
TP_ThirdPersonBP - Added missing SoftCollision collision channel
TP_ThirdPersonBP - Modified soft collision sidescroller platforms to match C++ version
TP_ThirdPersonBP - Removed combo attack released touch control input that could trigger incorrect controls
TP_ThirdPersonBP - Zeroed out Y component of SideScroller NPC knockback
TP_ThirdPersonBP - Fixed incorrectly configured wait for landing task on combat StateTree
TP_VehicleAdvBP - Restored HUD class on base GameMode, reparented OffroadGameMode to GameModeBase
TP_VehicleAdvBP - Reparented OffroadGameMode to GameModeBase
TP_VehicleAdvBP - Formatted Time Trial UI timer
TP_FirstPerson - Renamed input action asset throwing security warning in build
TP_FirstPerson - Shooter StateTree adjustments
TP_FirstPerson - Fixed a potential crash if the Shooter player was killed while unarmed
TP_FirstPersonBP - Renamed input action asset throwing security warning in build
TP_FirstPersonBP - Shooter StateTree adjustments
TP_FirstPersonBP - Removed unused AI assets
TP_FirstPersonBP - Added Asset Guidelines to Shooter Variant

#rnx

diff --git a/Templates/TP_FirstPerson/Source/TP_FirstPerson/Variant_Shooter/ShooterCharacter.cpp b/Templates/TP_FirstPerson/Source/TP_FirstPerson/Variant_Shooter/ShooterCharacter.cpp
--- a/Templates/TP_FirstPerson/Source/TP_FirstPerson/Variant_Shooter/ShooterCharacter.cpp
+++ b/Templates/TP_FirstPerson/Source/TP_FirstPerson/Variant_Shooter/ShooterCharacter.cpp
@@ -53,7 +53,11 @@ float AShooterCharacter::TakeDamage(float Damage, struct FDamageEvent const& Dam
 	if (CurrentHP <= 0.0f)
 	{
 		// deactivate the weapon
-		CurrentWeapon->DeactivateWeapon();
+		if (IsValid(CurrentWeapon))
+		{
+			CurrentWeapon->DeactivateWeapon();
+		}
+		
 
 		// reset the bullet counter UI
 		OnBulletCountUpdated.Broadcast(0, 0);
diff --git a/Templates/TP_ThirdPerson/Source/TP_ThirdPerson/TP_ThirdPersonCharacter.cpp b/Templates/TP_ThirdPerson/Source/TP_ThirdPerson/TP_ThirdPersonCharacter.cpp
--- a/Templates/TP_ThirdPerson/Source/TP_ThirdPerson/TP_ThirdPersonCharacter.cpp
+++ b/Templates/TP_ThirdPerson/Source/TP_ThirdPerson/TP_ThirdPersonCharacter.cpp
@@ -29,7 +29,7 @@ ATP_ThirdPersonCharacter::ATP_ThirdPersonCharacter()
 
 	// Note: For faster iteration times these variables, and many more, can be tweaked in the Character Blueprint
 	// instead of recompiling to adjust them
-	GetCharacterMovement()->JumpZVelocity = 700.f;
+	GetCharacterMovement()->JumpZVelocity = 500.f;
 	GetCharacterMovement()->AirControl = 0.35f;
 	GetCharacterMovement()->MaxWalkSpeed = 500.f;
 	GetCharacterMovement()->MinAnalogWalkSpeed = 20.f;
diff --git a/Templates/TP_VehicleAdv/Source/TP_VehicleAdv/Variant_TimeTrial/TimeTrialUI.cpp b/Templates/TP_VehicleAdv/Source/TP_VehicleAdv/Variant_TimeTrial/TimeTrialUI.cpp
--- a/Templates/TP_VehicleAdv/Source/TP_VehicleAdv/Variant_TimeTrial/TimeTrialUI.cpp
+++ b/Templates/TP_VehicleAdv/Source/TP_VehicleAdv/Variant_TimeTrial/TimeTrialUI.cpp
@@ -19,10 +19,13 @@ void UTimeTrialUI::NativeConstruct()
 	StartUI->StartCountdown();
 }
 
-void UTimeTrialUI::UpdateLapCount(int32 Lap, float LapStartTime)
+void UTimeTrialUI::UpdateLapCount(int32 Lap, float NewLapStartTime)
 {
+	// save the new lap start time
+	LapStartTime = NewLapStartTime;
+
 	// calculate the lap time
-	const float LapTime = LapStartTime - LastLapTime;
+	const float LapTime = NewLapStartTime - LastLapTime;
 
 	// is this the first lap?
 	if (Lap > 1)
@@ -55,7 +58,7 @@ void UTimeTrialUI::UpdateLapCount(int32 Lap, float LapStartTime)
 	CurrentLap = Lap;
 
 	// save the lap start time
-	LastLapTime = LapStartTime;
+	LastLapTime = NewLapStartTime;
 
 	// pass control to BP to update the widgets
 	BP_UpdateLaps();
diff --git a/Templates/TP_VehicleAdv/Source/TP_VehicleAdv/Variant_TimeTrial/TimeTrialUI.h b/Templates/TP_VehicleAdv/Source/TP_VehicleAdv/Variant_TimeTrial/TimeTrialUI.h
--- a/Templates/TP_VehicleAdv/Source/TP_VehicleAdv/Variant_TimeTrial/TimeTrialUI.h
+++ b/Templates/TP_VehicleAdv/Source/TP_VehicleAdv/Variant_TimeTrial/TimeTrialUI.h
@@ -32,6 +32,9 @@ protected:
 	/** Best lap time, in seconds */
 	float BestLapTime = 0.0f;
 
+	/** Game time when this lap started */
+	float LapStartTime = 0.0f;
+
 	/** Current lap number */
 	int32 CurrentLap = 0;
 
@@ -48,7 +51,7 @@ protected:
 public:
 
 	/** Increments the lap and updates the lap counter */
-	void UpdateLapCount(int32 Lap, float LapStartTime);
+	void UpdateLapCount(int32 Lap, float NewLapStartTime);
 
 	/** Allows Blueprint control to update the lap tracker widgets */
 	UFUNCTION(BlueprintImplementableEvent, Category="Time Trial", meta=(DisplayName = "Update Laps"))
@@ -68,5 +71,7 @@ protected:
 	UFUNCTION(BlueprintPure, Category="Time Trial")
 	float GetBestLapTime() const { return BestLapTime; };
 
-
+	/** Gets the best lap time saved */
+	UFUNCTION(BlueprintPure, Category="Time Trial")
+	float GetLapStartTime() const { return LapStartTime; };
 };
