Templates
- TP_VehicleAdv: Implemented automatic vehicle respawning on destruction
- TP_VehicleAdvBP: Renamed BP assets to match template naming conventions
- TP_VehicleAdvBP: Moved all widget management to PlayerController classes
- TP_VehicleAdvBP: Fixed a scaling issue on TimeTrial PlayerStart that caused respawned cars to appear squished.

#rnx

diff --git a/Templates/TP_VehicleAdv/Source/TP_VehicleAdv/TP_VehicleAdvPlayerController.cpp b/Templates/TP_VehicleAdv/Source/TP_VehicleAdv/TP_VehicleAdvPlayerController.cpp
--- a/Templates/TP_VehicleAdv/Source/TP_VehicleAdv/TP_VehicleAdvPlayerController.cpp
+++ b/Templates/TP_VehicleAdv/Source/TP_VehicleAdv/TP_VehicleAdvPlayerController.cpp
@@ -8,6 +8,8 @@
 #include "ChaosWheeledVehicleMovementComponent.h"
 #include "Blueprint/UserWidget.h"
 #include "TP_VehicleAdv.h"
+#include "Kismet/GameplayStatics.h"
+#include "GameFramework/PlayerStart.h"
 
 void ATP_VehicleAdvPlayerController::BeginPlay()
 {
@@ -93,4 +95,26 @@ void ATP_VehicleAdvPlayerController::OnPossess(APawn* InPawn)
 
 	// get a pointer to the controlled pawn
 	VehiclePawn = CastChecked<ATP_VehicleAdvPawn>(InPawn);
+
+	// subscribe to the pawn's OnDestroyed delegate
+	VehiclePawn->OnDestroyed.AddDynamic(this, &ATP_VehicleAdvPlayerController::OnPawnDestroyed);
+}
+
+void ATP_VehicleAdvPlayerController::OnPawnDestroyed(AActor* DestroyedPawn)
+{
+	// find the player start
+	TArray<AActor*> ActorList;
+	UGameplayStatics::GetAllActorsOfClass(GetWorld(), APlayerStart::StaticClass(), ActorList);
+
+	if (ActorList.Num() > 0)
+	{
+		// spawn a vehicle at the player start
+		const FTransform SpawnTransform = ActorList[0]->GetActorTransform();
+
+		if (ATP_VehicleAdvPawn* RespawnedVehicle = GetWorld()->SpawnActor<ATP_VehicleAdvPawn>(VehiclePawnClass, SpawnTransform))
+		{
+			// possess the vehicle
+			Possess(RespawnedVehicle);
+		}
+	}
 }
diff --git a/Templates/TP_VehicleAdv/Source/TP_VehicleAdv/TP_VehicleAdvPlayerController.h b/Templates/TP_VehicleAdv/Source/TP_VehicleAdv/TP_VehicleAdvPlayerController.h
--- a/Templates/TP_VehicleAdv/Source/TP_VehicleAdv/TP_VehicleAdvPlayerController.h
+++ b/Templates/TP_VehicleAdv/Source/TP_VehicleAdv/TP_VehicleAdvPlayerController.h
@@ -50,6 +50,10 @@ protected:
 	UPROPERTY(EditAnywhere, Category = "Input|Steering Wheel Controls", meta = (EditCondition = "bUseSteeringWheelControls"))
 	UInputMappingContext* SteeringWheelInputMappingContext;
 
+	/** Type of vehicle to automatically respawn when it's destroyed */
+	UPROPERTY(EditAnywhere, Category="Vehicle|Respawn")
+	TSubclassOf<ATP_VehicleAdvPawn> VehiclePawnClass;
+
 	/** Pointer to the controlled vehicle pawn */
 	TObjectPtr<ATP_VehicleAdvPawn> VehiclePawn;
 
@@ -59,25 +63,26 @@ protected:
 
 	/** Pointer to the UI widget */
 	TObjectPtr<UTP_VehicleAdvUI> VehicleUI;
-
-	
-
-	// Begin Actor interface
+		
 protected:
 
+	/** Gameplay initialization */
 	virtual void BeginPlay() override;
+
+	/** Input setup */
 	virtual void SetupInputComponent() override;
 
 public:
 
+	/** Update vehicle UI on tick */
 	virtual void Tick(float Delta) override;
 
-	// End Actor interface
-
-	// Begin PlayerController interface
 protected:
 
+	/** Pawn setup */
 	virtual void OnPossess(APawn* InPawn) override;
 
-	// End PlayerController interface
+	/** Handles pawn destruction and respawning */
+	UFUNCTION()
+	void OnPawnDestroyed(AActor* DestroyedPawn);
 };
diff --git a/Templates/TP_VehicleAdv/Source/TP_VehicleAdv/Variant_TimeTrial/TimeTrialPlayerController.cpp b/Templates/TP_VehicleAdv/Source/TP_VehicleAdv/Variant_TimeTrial/TimeTrialPlayerController.cpp
--- a/Templates/TP_VehicleAdv/Source/TP_VehicleAdv/Variant_TimeTrial/TimeTrialPlayerController.cpp
+++ b/Templates/TP_VehicleAdv/Source/TP_VehicleAdv/Variant_TimeTrial/TimeTrialPlayerController.cpp
@@ -14,6 +14,8 @@
 #include "ChaosWheeledVehicleMovementComponent.h"
 #include "Blueprint/UserWidget.h"
 #include "TP_VehicleAdv.h"
+#include "Kismet/GameplayStatics.h"
+#include "GameFramework/PlayerStart.h"
 
 void ATimeTrialPlayerController::BeginPlay()
 {
@@ -106,8 +108,14 @@ void ATimeTrialPlayerController::OnPossess(APawn* InPawn)
 	// get a pointer to the controlled pawn
 	VehiclePawn = CastChecked<ATP_VehicleAdvPawn>(InPawn);
 
-	// disable input on the pawn
-	VehiclePawn->DisableInput(this);
+	// subscribe to the pawn's OnDestroyed delegate
+	VehiclePawn->OnDestroyed.AddDynamic(this, &ATimeTrialPlayerController::OnPawnDestroyed);
+
+	// disable input on the pawn if the race hasn't started yet
+	if (!bRaceStarted)
+	{
+		VehiclePawn->DisableInput(this);
+	}	
 }
 
 void ATimeTrialPlayerController::Tick(float Delta)
@@ -129,6 +137,9 @@ void ATimeTrialPlayerController::StartRace()
 		SetTargetGate(GM->GetFinishLine()->GetNextMarker());
 	}
 
+	// raise the race started flag so any respawned vehicles start with controls unlocked 
+	bRaceStarted = true;
+
 	// start the first lap
 	CurrentLap = 0;
 	IncrementLapCount();
@@ -155,3 +166,22 @@ void ATimeTrialPlayerController::SetTargetGate(ATimeTrialTrackGate* Gate)
 {
 	TargetGate = Gate;
 }
+
+void ATimeTrialPlayerController::OnPawnDestroyed(AActor* DestroyedPawn)
+{
+	// find the player start
+	TArray<AActor*> ActorList;
+	UGameplayStatics::GetAllActorsOfClass(GetWorld(), APlayerStart::StaticClass(), ActorList);
+
+	if (ActorList.Num() > 0)
+	{
+		// spawn a vehicle at the player start
+		const FTransform SpawnTransform = ActorList[0]->GetActorTransform();
+
+		if (ATP_VehicleAdvPawn* RespawnedVehicle = GetWorld()->SpawnActor<ATP_VehicleAdvPawn>(VehiclePawnClass, SpawnTransform))
+		{
+			// possess the vehicle
+			Possess(RespawnedVehicle);
+		}
+	}
+}
diff --git a/Templates/TP_VehicleAdv/Source/TP_VehicleAdv/Variant_TimeTrial/TimeTrialPlayerController.h b/Templates/TP_VehicleAdv/Source/TP_VehicleAdv/Variant_TimeTrial/TimeTrialPlayerController.h
--- a/Templates/TP_VehicleAdv/Source/TP_VehicleAdv/Variant_TimeTrial/TimeTrialPlayerController.h
+++ b/Templates/TP_VehicleAdv/Source/TP_VehicleAdv/Variant_TimeTrial/TimeTrialPlayerController.h
@@ -71,6 +71,13 @@ protected:
 	/** Lap counter */
 	int32 CurrentLap = 0;
 
+	/** If true, the race has already started */
+	bool bRaceStarted = false;
+
+	/** Type of vehicle to automatically respawn when it's destroyed */
+	UPROPERTY(EditAnywhere, Category="Vehicle|Respawn")
+	TSubclassOf<ATP_VehicleAdvPawn> VehiclePawnClass;
+
 	/** Pointer to the controlled vehicle pawn */
 	TObjectPtr<ATP_VehicleAdvPawn> VehiclePawn;
 
@@ -104,4 +111,10 @@ public:
 
 	/** Sets the target track gate for this player */
 	void SetTargetGate(ATimeTrialTrackGate* Gate);
+
+protected:
+
+	/** Handles pawn destruction and respawning */
+	UFUNCTION()
+	void OnPawnDestroyed(AActor* DestroyedPawn);
 };
diff --git a/Templates/TP_VehicleAdvBP/Config/DefaultEngine.ini b/Templates/TP_VehicleAdvBP/Config/DefaultEngine.ini
--- a/Templates/TP_VehicleAdvBP/Config/DefaultEngine.ini
+++ b/Templates/TP_VehicleAdvBP/Config/DefaultEngine.ini
@@ -15,7 +15,7 @@ ThreePlayerSplitscreenLayout=FavorTop
 GameInstanceClass=/Script/Engine.GameInstance
 GameDefaultMap=/Game/VehicleTemplate/Maps/Lvl_VehicleBasic.Lvl_VehicleBasic
 ServerDefaultMap=/Engine/Maps/Entry.Entry
-GlobalDefaultGameMode=/Game/VehicleTemplate/Blueprints/BP_GM_VehicleBasic.BP_GM_VehicleBasic_C
+GlobalDefaultGameMode=/Game/VehicleTemplate/Blueprints/BP_VehicleAdvGameMode.BP_VehicleAdvGameMode_C
 GlobalDefaultServerGameMode=None
 
 [/Script/Engine.RendererSettings]
@@ -82,3 +82,18 @@ AppliedDefaultGraphicsPerformance=Maximum
 [/Script/Engine.Engine]
 +ActiveGameNameRedirects=(OldGameName="TP_VehicleAdvBP",NewGameName="/Script/VehicleTemplateNew")
 +ActiveGameNameRedirects=(OldGameName="/Script/TP_VehicleAdvBP",NewGameName="/Script/VehicleTemplateNew")
+
+[/Script/AndroidFileServerEditor.AndroidFileServerRuntimeSettings]
+bEnablePlugin=True
+bAllowNetworkConnection=True
+SecurityToken=1EF3860741EAFFC42C80FFA0493E9638
+bIncludeInShipping=False
+bAllowExternalStartInShipping=False
+bCompileAFSProject=False
+bUseCompression=False
+bLogFiles=False
+bReportStats=False
+ConnectionType=USBOnly
+bUseManualIPAddress=False
+ManualIPAddress=
+
diff --git a/Templates/TP_VehicleAdvBP/Config/TemplateDefs.ini b/Templates/TP_VehicleAdvBP/Config/TemplateDefs.ini
--- a/Templates/TP_VehicleAdvBP/Config/TemplateDefs.ini
+++ b/Templates/TP_VehicleAdvBP/Config/TemplateDefs.ini
@@ -42,8 +42,8 @@ EditDetailLevelPreference="High"
 
 ; Variant packs
 ; NOTE: make sure these are commented out for release!
-;SharedContentPacks=(MountName="Variant_Offroad",DetailLevels=("Standard"))
-;SharedContentPacks=(MountName="Variant_Timetrial",DetailLevels=("Standard"))
+SharedContentPacks=(MountName="Variant_Offroad",DetailLevels=("Standard"))
+SharedContentPacks=(MountName="Variant_Timetrial",DetailLevels=("Standard"))
 
 Variants=(Name="Timetrial",LocalizedDisplayNames=((Language="en",Text="Timetrial")),LocalizedDescriptions=((Language="en",Text="A timed lap with UI and track markers.")),SharedContentPacks=((DetailLevels=(Standard),MountName="Variant_Timetrial")))
 
