TP_VehicleAdv, TP_VehicleAdvBP
- Implemented missing auto-flip check on C++ template.
- Fixed camera lag inconsistencies between BP and C++ template versions.
- Fixed incorrect torque curve on offroad vehicle on BP template.
- PlayerController now attaches to vehicle, fixing some incorrect WP streaming issues.
- Enabled WP streaming on offroad map for C++ template.

#rnx

#virtualized

diff --git a/Templates/TP_VehicleAdv/Source/TP_VehicleAdv/TP_VehicleAdvPawn.cpp b/Templates/TP_VehicleAdv/Source/TP_VehicleAdv/TP_VehicleAdvPawn.cpp
--- a/Templates/TP_VehicleAdv/Source/TP_VehicleAdv/TP_VehicleAdvPawn.cpp
+++ b/Templates/TP_VehicleAdv/Source/TP_VehicleAdv/TP_VehicleAdvPawn.cpp
@@ -11,6 +11,7 @@
 #include "InputActionValue.h"
 #include "ChaosWheeledVehicleMovementComponent.h"
 #include "TP_VehicleAdv.h"
+#include "TimerManager.h"
 
 #define LOCTEXT_NAMESPACE "VehiclePawn"
 
@@ -91,6 +92,22 @@ void ATP_VehicleAdvPawn::SetupPlayerInputComponent(class UInputComponent* Player
 	}
 }
 
+void ATP_VehicleAdvPawn::BeginPlay()
+{
+	Super::BeginPlay();
+
+	// set up the flipped check timer
+	GetWorld()->GetTimerManager().SetTimer(FlipCheckTimer, this, &ATP_VehicleAdvPawn::FlippedCheck, FlipCheckTime, true);
+}
+
+void ATP_VehicleAdvPawn::EndPlay(EEndPlayReason::Type EndPlayReason)
+{
+	// clear the flipped check timer
+	GetWorld()->GetTimerManager().ClearTimer(FlipCheckTimer);
+
+	Super::EndPlay(EndPlayReason);
+}
+
 void ATP_VehicleAdvPawn::Tick(float Delta)
 {
 	Super::Tick(Delta);
@@ -255,4 +272,28 @@ void ATP_VehicleAdvPawn::DoResetVehicle()
 	GetMesh()->SetPhysicsLinearVelocity(FVector::ZeroVector);
 }
 
+void ATP_VehicleAdvPawn::FlippedCheck()
+{
+	// check the difference in angle between the mesh's up vector and world up
+	const float UpDot = FVector::DotProduct(FVector::UpVector, GetMesh()->GetUpVector());
+
+	if (UpDot < FlipCheckMinDot)
+	{
+		// is this the second time we've checked that the vehicle is still flipped?
+		if (bPreviousFlipCheck)
+		{
+			// reset the vehicle to upright
+			DoResetVehicle();
+		}
+		
+		// set the flipped check flag so the next check resets the car
+		bPreviousFlipCheck = true;
+
+	} else {
+
+		// we're upright. reset the flipped check flag
+		bPreviousFlipCheck = false;
+	}
+}
+
 #undef LOCTEXT_NAMESPACE
\ No newline at end of file
diff --git a/Templates/TP_VehicleAdv/Source/TP_VehicleAdv/TP_VehicleAdvPawn.h b/Templates/TP_VehicleAdv/Source/TP_VehicleAdv/TP_VehicleAdvPawn.h
--- a/Templates/TP_VehicleAdv/Source/TP_VehicleAdv/TP_VehicleAdvPawn.h
+++ b/Templates/TP_VehicleAdv/Source/TP_VehicleAdv/TP_VehicleAdvPawn.h
@@ -76,6 +76,20 @@ protected:
 	/** Keeps track of which camera is active */
 	bool bFrontCameraActive = false;
 
+	/** Keeps track of whether the car is flipped. If this is true for two flip checks, resets the vehicle automatically */
+	bool bPreviousFlipCheck = false;
+
+	/** Time between automatic flip checks */
+	UPROPERTY(EditAnywhere, Category="Flip Check", meta = (Units = "s"))
+	float FlipCheckTime = 3.0f;
+
+	/** Minimum dot product value for the vehicle's up direction that we still consider upright */
+	UPROPERTY(EditAnywhere, Category="Flip Check")
+	float FlipCheckMinDot = -0.2f;
+
+	/** Flip check timer */
+	FTimerHandle FlipCheckTimer;
+
 public:
 	ATP_VehicleAdvPawn();
 
@@ -87,6 +101,13 @@ public:
 
 	// Begin Actor interface
 
+	/** Initialization */
+	virtual void BeginPlay() override;
+
+	/** Cleanup */
+	virtual void EndPlay(EEndPlayReason::Type EndPlayReason) override;
+
+	/** Update */
 	virtual void Tick(float Delta) override;
 
 	// End Actor interface
@@ -167,6 +188,10 @@ protected:
 	UFUNCTION(BlueprintImplementableEvent, Category="Vehicle")
 	void BrakeLights(bool bBraking);
 
+	/** Checks if the car is flipped upside down and automatically resets it */
+	UFUNCTION()
+	void FlippedCheck();
+
 public:
 	/** Returns the front spring arm subobject */
 	FORCEINLINE USpringArmComponent* GetFrontSpringArm() const { return FrontSpringArm; }
diff --git a/Templates/TP_VehicleAdv/Source/TP_VehicleAdv/TP_VehicleAdvPlayerController.cpp b/Templates/TP_VehicleAdv/Source/TP_VehicleAdv/TP_VehicleAdvPlayerController.cpp
--- a/Templates/TP_VehicleAdv/Source/TP_VehicleAdv/TP_VehicleAdvPlayerController.cpp
+++ b/Templates/TP_VehicleAdv/Source/TP_VehicleAdv/TP_VehicleAdvPlayerController.cpp
@@ -16,6 +16,9 @@ void ATP_VehicleAdvPlayerController::BeginPlay()
 {
 	Super::BeginPlay();
 	
+	// ensure we're attached to the vehicle pawn so that World Partition streaming works correctly
+	bAttachToPawn = true;
+
 	// only spawn UI on local player controllers
 	if (IsLocalPlayerController())
 	{
