TP_VehicleAdv
- Added GC safeguards to untracked UMG widget pointers.
#rnx

diff --git a/Templates/TP_VehicleAdv/Source/TP_VehicleAdv/TP_VehicleAdvPlayerController.h b/Templates/TP_VehicleAdv/Source/TP_VehicleAdv/TP_VehicleAdvPlayerController.h
--- a/Templates/TP_VehicleAdv/Source/TP_VehicleAdv/TP_VehicleAdvPlayerController.h
+++ b/Templates/TP_VehicleAdv/Source/TP_VehicleAdv/TP_VehicleAdvPlayerController.h
@@ -63,6 +63,7 @@ protected:
 	TSubclassOf<UTP_VehicleAdvUI> VehicleUIClass;
 
 	/** Pointer to the UI widget */
+	UPROPERTY()
 	TObjectPtr<UTP_VehicleAdvUI> VehicleUI;
 		
 protected:
diff --git a/Templates/TP_VehicleAdv/Source/TP_VehicleAdv/Variant_TimeTrial/TimeTrialPlayerController.h b/Templates/TP_VehicleAdv/Source/TP_VehicleAdv/Variant_TimeTrial/TimeTrialPlayerController.h
--- a/Templates/TP_VehicleAdv/Source/TP_VehicleAdv/Variant_TimeTrial/TimeTrialPlayerController.h
+++ b/Templates/TP_VehicleAdv/Source/TP_VehicleAdv/Variant_TimeTrial/TimeTrialPlayerController.h
@@ -57,6 +57,7 @@ protected:
 	TSubclassOf<UTimeTrialUI> UIWidgetClass;
 
 	/** Pointer to the UI Widget */
+	UPROPERTY()
 	TObjectPtr<UTimeTrialUI> UIWidget;
 
 	/** Type of the UI to spawn */
@@ -64,6 +65,7 @@ protected:
 	TSubclassOf<UTP_VehicleAdvUI> VehicleUIClass;
 
 	/** Pointer to the UI widget */
+	UPROPERTY()
 	TObjectPtr<UTP_VehicleAdvUI> VehicleUI;
 
 	/** Next track gate the car should pass */
