Templates
- Fixed a typecasing issue causing build failures on Mac

#rnx

#Synced-CL 47537391

diff --git a/Templates/TP_VehicleAdv/Source/TP_VehicleAdv/TP_VehicleAdv.Build.cs b/Templates/TP_VehicleAdv/Source/TP_VehicleAdv/TP_VehicleAdv.Build.cs
--- a/Templates/TP_VehicleAdv/Source/TP_VehicleAdv/TP_VehicleAdv.Build.cs
+++ b/Templates/TP_VehicleAdv/Source/TP_VehicleAdv/TP_VehicleAdv.Build.cs
@@ -24,7 +24,7 @@ public class TP_VehicleAdv : ModuleRules
 			"TP_VehicleAdv",
 			"TP_VehicleAdv/SportsCar",
 			"TP_VehicleAdv/OffroadCar",
-			"TP_VehicleAdv/Variant_Offroad",
+			"TP_VehicleAdv/Variant_OffRoad",
 			"TP_VehicleAdv/Variant_TimeTrial",
 			"TP_VehicleAdv/Variant_TimeTrial/UI"
 		});
